Make AutomatedPerfTesting plugin generally available.

#virtualized

diff --git a/Engine/Plugins/Performance/AutomatedPerfTesting/AutomatedPerfTesting.uplugin b/Engine/Plugins/Performance/AutomatedPerfTesting/AutomatedPerfTesting.uplugin
new file mode 100644
--- /dev/null
+++ b/Engine/Plugins/Performance/AutomatedPerfTesting/AutomatedPerfTesting.uplugin
@@ -0,0 +1,30 @@
+{
+	"FileVersion": 3,
+	"Version": 1,
+	"VersionName": "1.0",
+	"FriendlyName": "AutomatedPerfTesting",
+	"Description": "",
+	"Category": "Testing",
+	"CreatedBy": "",
+	"CreatedByURL": "",
+	"DocsURL": "",
+	"MarketplaceURL": "",
+	"SupportURL": "",
+	"CanContainContent": true,
+	"IsBetaVersion": false,
+	"IsExperimentalVersion": true,
+	"Installed": false,
+	"Modules": [
+		{
+			"Name": "AutomatedPerfTesting",
+			"Type": "Runtime",
+			"LoadingPhase": "Default"
+		}
+	],
+	"Plugins": [
+		{
+			"Name": "Gauntlet",
+			"Enabled": true
+		}
+	]
+}
\ No newline at end of file
diff --git a/Engine/Plugins/Performance/AutomatedPerfTesting/Build/Scripts/AutomatedPerfTestConfig.cs b/Engine/Plugins/Performance/AutomatedPerfTesting/Build/Scripts/AutomatedPerfTestConfig.cs
new file mode 100644
--- /dev/null
+++ b/Engine/Plugins/Performance/AutomatedPerfTesting/Build/Scripts/AutomatedPerfTestConfig.cs
@@ -0,0 +1,87 @@
+using System;
+using System.Collections.Generic;
+using System.IO;
+using System.Linq;
+using EpicGame;
+using Gauntlet;
+
+namespace AutomatedPerfTest
+{
+	public class AutomatedPerfTestConfigBase : UnrealTestConfiguration
+	{
+		/// <summary>
+		/// If we're running on the build machine
+		/// </summary>
+		[AutoParamWithNames(false, "AutomatedPerfTest.IsBuildMachine")]
+		public bool IsBuildMachine;
+		
+		/// <summary>
+		/// Name of the test, useful for identifying it later
+		/// </summary>
+		[AutoParamWithNames("AutomatedPerfTest.TestName")]
+		public string TestName = "";
+		
+		/// <summary>
+		/// If set, will prepend platform name and use this device profile instead of the default
+		/// </summary>
+		[AutoParamWithNames("AutomatedPerfTest.DeviceProfileOverride")]
+		public string DeviceProfileOverride = "";
+		
+		/// <summary>
+		/// If we're running on the build machine
+		/// </summary>
+		[AutoParamWithNames(false, "AutomatedPerfTest.DoInsightsTrace")]
+		public bool DoInsightsTrace;
+		
+		/// <summary>
+		/// Which trace channels to test with
+		/// </summary>
+		[AutoParamWithNames("AutomatedPerfTest.TraceChannels")]
+		public string TraceChannels = "default,screenshot";
+		
+		/// <summary>
+		/// If we're running on the build machine
+		/// </summary>
+		[AutoParamWithNames(false, "AutomatedPerfTest.DoCSVProfiler")]
+		public bool DoCSVProfiler;
+		
+		/// <summary>
+		/// If we're running on the build machine
+		/// </summary>
+		[AutoParamWithNames(false, "AutomatedPerfTest.DoFPSChart")]
+		public bool DoFPSChart;
+		
+		/// <summary>
+		/// Let BuildGraph tell us where we should output the Insights trace after running a test so that we know where to grab it from when we're done
+		/// </summary>
+		[AutoParamWithNames("", "AutomatedPerfTest.PerfCacheRoot")]
+		public string PerfCacheRoot;
+		
+		/// <summary>
+		/// Which platform we're testing on 
+		/// </summary>
+		[AutoParamWithNames("", "AutomatedPerfTest.TestPlatform")]
+		public string TestPlatform;
+		
+		/// <summary>
+		/// Let BuildGraph tell us where we should output the Insights trace after running a test so that we know where to grab it from when we're done
+		/// </summary>
+		[AutoParamWithNames(false, "AutomatedPerfTest.DoVideoCapture")]
+		public bool DoVideoCapture;
+	}
+
+	public class AutomatedSequencePerfTestConfig : AutomatedPerfTestConfigBase
+	{
+		/// <summary>
+		/// Which map to run the test on
+		/// </summary>
+		[AutoParamWithNames("", "AutomatedPerfTest.TestMap")]
+		public string TestMap;
+		
+		/// <summary>
+		/// Which sequence to run the test with
+		/// </summary>
+		[AutoParamWithNames("", "AutomatedPerfTest.SequencePath")]
+		public string SequencePath;
+	}
+}
diff --git a/Engine/Plugins/Performance/AutomatedPerfTesting/Source/AutomatedPerfTesting/AutomatedPerfTesting.Build.cs b/Engine/Plugins/Performance/AutomatedPerfTesting/Source/AutomatedPerfTesting/AutomatedPerfTesting.Build.cs
new file mode 100644
--- /dev/null
+++ b/Engine/Plugins/Performance/AutomatedPerfTesting/Source/AutomatedPerfTesting/AutomatedPerfTesting.Build.cs
@@ -0,0 +1,53 @@
+// Copyright Epic Games, Inc. All Rights Reserved.
+
+using UnrealBuildTool;
+
+public class AutomatedPerfTesting : ModuleRules
+{
+	public AutomatedPerfTesting(ReadOnlyTargetRules Target) : base(Target)
+	{
+		PCHUsage = ModuleRules.PCHUsageMode.UseExplicitOrSharedPCHs;
+		
+		PublicIncludePaths.AddRange(
+			new string[] {
+				// ... add public include paths required here ...
+			}
+			);
+				
+		
+		PrivateIncludePaths.AddRange(
+			new string[] {
+				// ... add other private include paths required here ...
+			}
+			);
+			
+		
+		PublicDependencyModuleNames.AddRange(
+			new string[]
+			{
+				"Core", "Gauntlet"
+				// ... add other public dependencies that you statically link with here ...
+			}
+			);
+			
+		
+		PrivateDependencyModuleNames.AddRange(
+			new string[]
+			{
+				"CoreUObject",
+				"Engine",
+				"Slate",
+				"SlateCore", "LevelSequence", "MovieScene"
+				// ... add private dependencies that you statically link with here ...	
+			}
+			);
+		
+		
+		DynamicallyLoadedModuleNames.AddRange(
+			new string[]
+			{
+				// ... add any modules that your module loads dynamically here ...
+			}
+			);
+	}
+}
diff --git a/Engine/Plugins/Performance/AutomatedPerfTesting/Source/AutomatedPerfTesting/Private/AutomatedPerfTestControllerBase.cpp b/Engine/Plugins/Performance/AutomatedPerfTesting/Source/AutomatedPerfTesting/Private/AutomatedPerfTestControllerBase.cpp
new file mode 100644
--- /dev/null
+++ b/Engine/Plugins/Performance/AutomatedPerfTesting/Source/AutomatedPerfTesting/Private/AutomatedPerfTestControllerBase.cpp
@@ -0,0 +1,593 @@
+﻿// Copyright Epic Games, Inc. All Rights Reserved.
+
+#include "AutomatedPerfTestControllerBase.h"
+
+#include <valarray>
+
+#include "AutomatedPerfTesting.h"
+#include "DeviceProfiles/DeviceProfileManager.h"
+#include "GameFramework/GameModeBase.h"
+#include "AutomatedPerfTestInterface.h"
+#include "Misc/App.h"
+#include "Misc/CommandLine.h"
+#include "ProfilingDebugging/TraceAuxiliary.h"
+#include "TimerManager.h"
+#include "Commandlets/Commandlet.h"
+#include "ProfilingDebugging/CsvProfiler.h"
+#include "Async/TaskGraphInterfaces.h"
+#include "VideoRecordingSystem.h"
+#include "PlatformFeatures.h"
+#include "Engine/Engine.h"
+#include "Misc/Paths.h"
+
+DEFINE_LOG_CATEGORY(LogAutomatedPerfTest)
+CSV_DEFINE_CATEGORY(AutomatedPerfTest, true);
+
+namespace AutomatedPerfTesting
+{
+	static TAutoConsoleVariable<FString> CVarExplicitTestID(
+		TEXT("AutomatedPerfTest.ExplicitTestID"),
+		"",
+		TEXT("Overrides the automatically generated test ID with the one provided"),
+		ECVF_Default);
+
+	static TAutoConsoleVariable<FString> CVarTraceChannels(
+		TEXT("AutomatedPerfTest.TraceChannels"),
+		"default,screenshot,stats",
+		TEXT("Sets the channels that will be traced when the AutomatedPerfTest is run. Defaults to default,screenshot,stats"),
+		ECVF_Default);
+
+	static TAutoConsoleVariable<FString> CVarDeviceProfileOverride(
+		TEXT("AutomatedPerfTest.DeviceProfileOverride"),
+		"",
+		TEXT("Overrides device profile used when running automated performance tests if set"),
+		ECVF_Default);
+
+	static TAutoConsoleVariable<FString> CVarTestName(
+		TEXT("AutomatedPerfTest.TestName"),
+		"Automated",
+		TEXT("User-defined ID of the test being run. Helpful for identifying test specifics eg Automated, TestingSomeVariable, MyCoolTest"),
+		ECVF_Default);
+}
+
+void UAutomatedPerfTestControllerBase::OnPreWorldInitializeInternal(UWorld* World, const UWorld::InitializationValues IVS)
+{
+	TryEarlyExec(World);
+	OnPreWorldInitialize(World);
+}
+
+void UAutomatedPerfTestControllerBase::OnPreWorldInitialize(UWorld* World)
+{
+	check(World);
+	World->GameStateSetEvent.AddUObject(this, &ThisClass::OnGameStateSet);
+	World->OnWorldBeginPlay.AddUObject(this, &ThisClass::OnWorldBeginPlay);
+}
+
+void UAutomatedPerfTestControllerBase::TryEarlyExec(UWorld* const World)
+{
+	check(World);
+	
+	if (GEngine)
+	{
+		// Search the list of deferred commands
+		const TArray<FString>& DeferredCmds = GEngine->DeferredCommands;
+		TArray<int32> ExecutedIndices;
+		for (int32 DeferredCmdIndex = 0; DeferredCmdIndex < DeferredCmds.Num(); ++DeferredCmdIndex)
+		{
+			// If the deferred command is one that should be executed early
+			const FString& DeferredCmd = DeferredCmds[DeferredCmdIndex];
+			if (CmdsToExecEarly.ContainsByPredicate([&DeferredCmd](const FString& CmdToFind) { return DeferredCmd.StartsWith(CmdToFind); }))
+			{
+				UE_LOG(LogAutomatedPerfTest, Log, TEXT("EarlyExec: executing '%s' early."), *DeferredCmd);
+				GEngine->Exec(World, *DeferredCmd);
+				ExecutedIndices.Push(DeferredCmdIndex);
+			}
+		}
+
+		// Remove the executed commands from the list of deferred commands
+		// Note: This is done in reverse order to ensure the cached indices remain valid
+		while (!ExecutedIndices.IsEmpty())
+		{
+			GEngine->DeferredCommands.RemoveAt(ExecutedIndices.Pop());
+		}
+	}
+	else
+	{
+		UE_LOG(LogAutomatedPerfTest, Warning, TEXT("Attempted EarlyExec without GEngine being ready"))
+	}
+}
+
+void UAutomatedPerfTestControllerBase::OnWorldBeginPlay()
+{
+	UE_LOG(LogAutomatedPerfTest, Log, TEXT("OnWorldBeginPlay"));
+	SetupTest();
+}
+
+void UAutomatedPerfTestControllerBase::OnGameStateSet(AGameStateBase* const GameStateBase)
+{
+	UE_LOG(LogAutomatedPerfTest, Log, TEXT("Game State Set"));
+	if (UWorld* const World = GetWorld())
+	{
+		World->GameStateSetEvent.RemoveAll(this);
+	}
+}
+
+UAutomatedPerfTestControllerBase::UAutomatedPerfTestControllerBase(const FObjectInitializer& ObjectInitializer)
+	: Super(ObjectInitializer)
+	, TraceChannels("default,screenshot,stats")
+	, bRequestsFPSChart(false)
+	, bRequestsInsightsTrace(false)
+	, bRequestsCSVProfiler(false)
+	, bRequestsVideoCapture(false)
+{
+	// cache this off once, so that it's consistent across the board
+	TestDatetime = FDateTime::Now().ToString(TEXT("%Y%m%d-%H%M%S")); 
+}
+
+FString UAutomatedPerfTestControllerBase::GetTestName()
+{
+	return AutomatedPerfTesting::CVarTestName->GetString();
+}
+
+FString UAutomatedPerfTestControllerBase::GetDeviceProfile()
+{
+	if(!AutomatedPerfTesting::CVarDeviceProfileOverride->GetString().IsEmpty())
+	{
+		return AutomatedPerfTesting::CVarDeviceProfileOverride->GetString();
+	}
+	
+	return UDeviceProfileManager::Get().GetActiveDeviceProfileName(); 
+}
+
+FString UAutomatedPerfTestControllerBase::GetExplicitTestID()
+{
+	// if the user has explicitly set a test ID, then return that
+	if(!AutomatedPerfTesting::CVarExplicitTestID->GetString().IsEmpty())
+	{
+		return AutomatedPerfTesting::CVarExplicitTestID->GetString();
+	}
+	
+	return TEXT("");
+}
+
+FString UAutomatedPerfTestControllerBase::GetTestID()
+{
+	// if the user has explicitly set a test ID, then return that
+	if(!AutomatedPerfTesting::CVarExplicitTestID->GetString().IsEmpty())
+	{
+		return AutomatedPerfTesting::CVarExplicitTestID->GetString();
+	}
+
+	const TArray<FString> TestCaseIDElements = {FApp::GetBuildVersion(),
+										  FPlatformProperties::PlatformName(),
+										  TestDatetime,
+								          *GetDeviceProfile(),
+								          *GetTestName()};
+	
+	// otherwise construct a unique ID of the form BuildVersion_PlatformName_YYYYMMDD-HHMMSS_DeviceProfile_TestName
+	FString TestCaseID = FString::Join(TestCaseIDElements, TEXT("_"));
+
+	return TestCaseID;
+}
+
+FString UAutomatedPerfTestControllerBase::GetOverallRegionName()
+{
+	return GetTestID() + "_" + "Overall";
+}
+
+bool UAutomatedPerfTestControllerBase::RequestsInsightsTrace() const
+{
+	return bRequestsInsightsTrace;
+}
+
+bool UAutomatedPerfTestControllerBase::RequestsCSVProfiler() const
+{
+	return bRequestsCSVProfiler;
+}
+
+bool UAutomatedPerfTestControllerBase::RequestsFPSChart() const
+{
+	return bRequestsFPSChart;
+}
+
+bool UAutomatedPerfTestControllerBase::RequestsVideoCapture() const
+{
+	return bRequestsVideoCapture;
+}
+
+bool UAutomatedPerfTestControllerBase::TryStartInsightsTrace()
+{
+	const FString TraceFileName = GetTestID() + ".utrace";
+	UE_LOG(LogAutomatedPerfTest, Log, TEXT("Attempting to start insights trace to file %s with channels %s"), *TraceFileName, *AutomatedPerfTesting::CVarTraceChannels->GetString());
+	return FTraceAuxiliary::Start(FTraceAuxiliary::EConnectionType::File, *TraceFileName, *AutomatedPerfTesting::CVarTraceChannels->GetString());
+}
+
+bool UAutomatedPerfTestControllerBase::TryStopInsightsTrace()
+{
+	if(FTraceAuxiliary::IsConnected())
+	{
+		return FTraceAuxiliary::Stop();
+	}
+	return false;
+}
+
+bool UAutomatedPerfTestControllerBase::TryStartCSVProfiler()
+{
+#if CSV_PROFILER
+	if(FCsvProfiler* const CsvProfiler = FCsvProfiler::Get())
+	{
+		const FString CSVFilename = GetTestID() + ".csv";
+
+		UE_LOG(LogAutomatedPerfTest, Log, TEXT("Attempting to start CSV Profile to file %s"), *CSVFilename);
+		
+		CsvProfiler->BeginCapture(-1, FString(), CSVFilename);
+		CsvProfiler->SetDeviceProfileName(GetDeviceProfile());
+		
+		return CsvProfiler->IsCapturing();
+	}
+#endif
+	UE_LOG(LogAutomatedPerfTest, Warning, TEXT("CSVProfiler Start requested, but not available."))
+	return false;
+}
+
+bool UAutomatedPerfTestControllerBase::TryStopCSVProfiler()
+{
+#if CSV_PROFILER
+	if(FCsvProfiler* const CsvProfiler = FCsvProfiler::Get())
+	{
+		const FGraphEventRef AutomatedPerfTestEndEvent = FGraphEvent::CreateGraphEvent();
+		CsvProfiler->EndCapture(AutomatedPerfTestEndEvent);
+		return true;
+	}
+#endif
+	UE_LOG(LogAutomatedPerfTest, Warning, TEXT("CSVProfiler Stop requested, but not available."))
+	return false;
+}
+
+bool UAutomatedPerfTestControllerBase::TryStartFPSChart()
+{
+	// don't open the folder the FPS chart gets sent to on exit, as it can cause issues when running unattended
+	GEngine->Exec(GetWorld(), TEXT("t.FPSChart.OpenFolderOnDump 0"));
+	GEngine->StartFPSChart(*GetOverallRegionName(), false);
+
+	return true;
+}
+
+bool UAutomatedPerfTestControllerBase::TryStopFPSChart()
+{
+	GEngine->StopFPSChart(*GetOverallRegionName());
+	return true;
+}
+
+bool UAutomatedPerfTestControllerBase::TryStartVideoCapture()
+{
+	if (IVideoRecordingSystem* const VideoRecordingSystem = IPlatformFeaturesModule::Get().GetVideoRecordingSystem())
+	{
+		const EVideoRecordingState RecordingState = VideoRecordingSystem->GetRecordingState();
+
+		if (RecordingState == EVideoRecordingState::None)
+		{
+			VideoRecordingSystem->EnableRecording(true);
+			
+			VideoRecordingTitle = FText::FromString(FPaths::Combine(FPaths::ProjectSavedDir(), GetTestID()));
+			const FVideoRecordingParameters VideoRecordingParameters(VideoRecordingSystem->GetMaximumRecordingSeconds(), true, false, false, FPlatformMisc::GetPlatformUserForUserIndex(0));
+			VideoRecordingSystem->NewRecording(*GetTestID(), VideoRecordingParameters);
+
+			if (VideoRecordingSystem->IsEnabled())
+			{
+				if (VideoRecordingSystem->GetRecordingState() == EVideoRecordingState::Starting || VideoRecordingSystem->GetRecordingState() == EVideoRecordingState::Recording)
+				{
+					UE_LOG(LogAutomatedPerfTest, Log, TEXT("Starting video recording %s..."), *GetTestID());
+					return true;
+				}
+				UE_LOG(LogAutomatedPerfTest, Warning, TEXT("Failed to start video recording %s. Current state is %i"), *GetTestID(), VideoRecordingSystem->GetRecordingState());
+			}
+			else
+			{
+				UE_LOG(LogAutomatedPerfTest, Warning, TEXT("Video recording could not be enabled."));
+			}
+		}
+		else
+		{
+			UE_LOG(LogAutomatedPerfTest, Warning, TEXT("Could not start a new recording, may be already recording."));
+		}
+	}
+	else
+	{
+		UE_LOG(LogAutomatedPerfTest, Warning, TEXT("Video recording system is null."));
+	}
+
+	return false;	
+}
+
+bool UAutomatedPerfTestControllerBase::TryFinalizingVideoCapture(const bool bStopAutoContinue/*=false*/)
+{
+	if (IVideoRecordingSystem* const VideoRecordingSystem = IPlatformFeaturesModule::Get().GetVideoRecordingSystem())
+	{
+		if (VideoRecordingSystem->GetRecordingState() != EVideoRecordingState::None)
+		{
+			VideoRecordingSystem->FinalizeRecording(true, VideoRecordingTitle, FText::GetEmpty(), bStopAutoContinue);
+
+			if (VideoRecordingSystem->GetRecordingState() == EVideoRecordingState::Finalizing)
+			{
+				UE_LOG(LogAutomatedPerfTest, Log, TEXT("Finalizing recording..."));
+				VideoRecordingSystem->GetOnVideoRecordingFinalizedDelegate().AddUObject(this, &ThisClass::OnVideoRecordingFinalized);
+				return true;
+			}
+			else
+			{
+				UE_LOG(LogAutomatedPerfTest, Warning, TEXT("Attempted to finalize video recording, but current state %i is not %i"), VideoRecordingSystem->GetRecordingState(), EVideoRecordingState::Finalizing)
+			}
+		}
+		else
+		{
+			UE_LOG(LogAutomatedPerfTest, Warning, TEXT("Attempted to finalize video recording, but state is %i"), VideoRecordingSystem->GetRecordingState())
+		}
+	}
+	else
+	{
+		UE_LOG(LogAutomatedPerfTest, Warning, TEXT("Video recording system is null."));
+	}
+
+	return false;	
+}
+
+void UAutomatedPerfTestControllerBase::SetupTest()
+{
+	UE_LOG(LogAutomatedPerfTest, Log, TEXT("Base:: SetupTest"));
+	if(RequestsInsightsTrace())
+	{
+		TryStartInsightsTrace();
+	}
+
+	if(RequestsCSVProfiler())
+	{
+		TryStartCSVProfiler();
+	}
+
+	if(RequestsFPSChart())
+	{
+		TryStartFPSChart();
+	}
+
+	if(RequestsVideoCapture())
+	{
+		TryStartVideoCapture();
+	}
+	
+	GameMode = GetWorld() ? GetWorld()->GetAuthGameMode() : NULL;
+
+	if(GameMode && GameMode->GetClass()->ImplementsInterface(UAutomatedPerfTestInterface::StaticClass()))
+	{
+		IAutomatedPerfTestInterface::Execute_SetupTest(GameMode);
+	}
+
+	// Subclasses should implement their own transitions from SetupTest to RunTest depending on their needs
+}
+
+void UAutomatedPerfTestControllerBase::RunTest()
+{
+	UE_LOG(LogAutomatedPerfTest, Log, TEXT("Base:: RunTest"));
+	if(RequestsInsightsTrace())
+	{
+		TRACE_BEGIN_REGION(*GetOverallRegionName());
+	}
+
+	if(RequestsCSVProfiler())
+	{
+		CSV_EVENT(AutomatedPerfTest, TEXT("START"), *GetOverallRegionName())
+	}
+	
+	if(GameMode && GameMode->GetClass()->ImplementsInterface(UAutomatedPerfTestInterface::StaticClass()))
+	{
+		IAutomatedPerfTestInterface::Execute_RunTest(GameMode);
+	}
+}
+
+void UAutomatedPerfTestControllerBase::TeardownTest()
+{
+	UE_LOG(LogAutomatedPerfTest, Log, TEXT("Base:: TeardownTest"));
+	
+	if(RequestsInsightsTrace())
+	{
+		TRACE_END_REGION(*GetOverallRegionName());
+	}
+	
+	if(RequestsCSVProfiler())
+	{
+		CSV_EVENT(AutomatedPerfTest, TEXT("END"), *GetOverallRegionName())
+		TryStopCSVProfiler();
+	}
+
+	if(RequestsFPSChart())
+	{
+		TryStopFPSChart();
+	}
+
+	if(RequestsVideoCapture())
+	{
+		if(!TryFinalizingVideoCapture())
+		{
+			UE_LOG(LogAutomatedPerfTest, Warning, TEXT("Attempted to finalize requested video capture, but failed."))
+		}
+	}
+
+	if(GameMode && GameMode->GetClass()->ImplementsInterface(UAutomatedPerfTestInterface::StaticClass()))
+	{
+		IAutomatedPerfTestInterface::Execute_TeardownTest(GameMode);
+	}
+
+
+	FTimerHandle UnusedHandle;
+	// TODO parameterize the teardown to Exit delay
+	GetWorld()->GetTimerManager().SetTimer(UnusedHandle, this, &UAutomatedPerfTestControllerBase::Exit, 1.0, false, 5.0);
+}
+
+void UAutomatedPerfTestControllerBase::Exit()
+{
+	if(RequestsInsightsTrace())
+	{
+		TryStopInsightsTrace();
+	}
+
+	if(GameMode && GameMode->GetClass()->ImplementsInterface(UAutomatedPerfTestInterface::StaticClass()))
+	{
+		IAutomatedPerfTestInterface::Execute_Exit(GameMode);
+	}
+
+	if(RequestsCSVProfiler())
+	{
+#if CSV_PROFILER
+		if(FCsvProfiler::Get()->IsWritingFile())
+		{
+			UE_LOG(LogAutomatedPerfTest, Log, TEXT("CSVProfile requested, and test is exiting, but CSV Profiler isn't done writing."))
+			// if we requested a CSV Profile, and the CSV Profiler is still writing the file, add a lambda to call the EndAutomatedPerfTest function
+			// so that we don't exit out of the application before the CSV Profiler is done
+			// TODO there's probably a nicer way to do this
+			// TODO might not need to do this if we set csv.BlockOnCaptureEnd
+			FCsvProfiler::Get()->OnCSVProfileFinished().AddLambda([this](const FString& Filename)
+			{
+				EndAutomatedPerfTest();
+			});
+			return;
+		}
+#endif
+	}
+	
+	EndAutomatedPerfTest();
+}
+
+void UAutomatedPerfTestControllerBase::OnInit()
+{
+	Super::OnInit();
+
+	UE_LOG(LogAutomatedPerfTest, Log, TEXT("Base:: OnInit"));
+
+	if (FParse::Value(FCommandLine::Get(), TEXT("AutomatedPerfTest.ExplicitTestID="), ExplicitTestID))
+	{
+		AutomatedPerfTesting::CVarExplicitTestID->Set(*ExplicitTestID);
+	}
+
+	// don't stop on separator because this will come in comma-separated
+	if (FParse::Value(FCommandLine::Get(), TEXT("AutomatedPerfTest.TraceChannels="), TraceChannels, false))
+	{
+		AutomatedPerfTesting::CVarTraceChannels->Set(*TraceChannels);
+	}
+
+	if (FParse::Value(FCommandLine::Get(), TEXT("AutomatedPerfTest.DeviceProfileOverride="), DeviceProfileOverride))
+	{
+		AutomatedPerfTesting::CVarDeviceProfileOverride->Set(*DeviceProfileOverride);
+	}
+
+	if (FString TestName; FParse::Value(FCommandLine::Get(), TEXT("AutomatedPerfTest.TestName="), TestID))
+	{
+	
+		AutomatedPerfTesting::CVarTestName->Set(*TestID);
+	}
+	
+	if(FParse::Param(FCommandLine::Get(), TEXT("AutomatedPerfTest.DoInsightsTrace")))
+	{
+		UE_LOG(LogAutomatedPerfTest, Log, TEXT("Insights Trace Requested"))
+		bRequestsInsightsTrace = true;
+	}
+	if(FParse::Param(FCommandLine::Get(), TEXT("AutomatedPerfTest.DoCSVProfiler")))
+	{
+		UE_LOG(LogAutomatedPerfTest, Log, TEXT("CSV Profiler Requested"))
+		bRequestsCSVProfiler = true;
+	}
+	if(FParse::Param(FCommandLine::Get(), TEXT("AutomatedPerfTest.DoFPSChart")))
+	{
+		UE_LOG(LogAutomatedPerfTest, Log, TEXT("FPSCharts Requested"))
+		bRequestsFPSChart = true;
+	}
+	if(FParse::Param(FCommandLine::Get(), TEXT("AutomatedPerfTest.DoVideoCapture")))
+	{
+		UE_LOG(LogAutomatedPerfTest, Log, TEXT("Video Capture Requested"))
+		bRequestsVideoCapture = true;
+	}
+	
+	FWorldDelegates::OnPreWorldInitialization.AddUObject(this, &ThisClass::OnPreWorldInitializeInternal);
+}
+
+void UAutomatedPerfTestControllerBase::OnTick(float TimeDelta)
+{
+	Super::OnTick(TimeDelta);
+	
+	MarkHeartbeatActive();
+}
+
+void UAutomatedPerfTestControllerBase::OnStateChange(FName OldState, FName NewState)
+{
+	Super::OnStateChange(OldState, NewState);
+}
+
+void UAutomatedPerfTestControllerBase::OnPreMapChange()
+{
+	Super::OnPreMapChange();
+}
+
+void UAutomatedPerfTestControllerBase::BeginDestroy()
+{
+	UnbindAllDelegates();
+	
+	Super::BeginDestroy();
+}
+
+void UAutomatedPerfTestControllerBase::EndAutomatedPerfTest(const int32 ExitCode)
+{
+	UnbindAllDelegates();
+
+	UE_LOG(LogAutomatedPerfTest, Log, TEXT("Test ID %s completed, requesting exit..."), *GetTestID());
+	
+	EndTest(ExitCode);
+}
+
+void UAutomatedPerfTestControllerBase::OnVideoRecordingFinalized(bool Succeeded, const FString& FilePath)
+{
+	if(!Succeeded)
+	{
+		UE_LOG(LogAutomatedPerfTest, Warning, TEXT("Video Capture finalized, but did not succeed"))
+	}
+/* TODO moving this to automation layer due to file access restrictions
+	FString SrcFilePath;
+	FString FileName;
+	FString Extension;
+	FPaths::Split(FilePath, SrcFilePath, FileName, Extension);
+
+	FString DestFileName = FileName + "." + Extension;
+	
+	const FString DestinationDir = FPaths::Combine(FPaths::ProjectSavedDir(), "Videos");
+	const FString DestinationFilePath = FPaths::Combine(DestinationDir, DestFileName);
+	
+	UE_LOG(LogAutomatedPerfTest, Log, TEXT("Copying video file %s to Saved: %s"), *FilePath, *DestinationFilePath);
+	
+	if(IFileManager::Get().Copy(*DestinationFilePath, *FilePath, 1, 1) != COPY_OK)
+	{
+		UE_LOG(LogAutomatedPerfTest, Warning, TEXT("Failed to copy video file"));
+	}
+	*/
+}
+
+void UAutomatedPerfTestControllerBase::UnbindAllDelegates()
+{
+	if(UWorld* const World = GetWorld())
+	{
+		World->OnWorldBeginPlay.RemoveAll(this);
+		World->GameStateSetEvent.RemoveAll(this);
+	}
+
+#if CSV_PROFILER
+	if (FCsvProfiler* const CsvProfiler = FCsvProfiler::Get())
+	{
+		CsvProfiler->OnCSVProfileFinished().Remove(CsvProfilerDelegateHandle);
+	}
+#endif // CSV_PROFILER
+
+	if(RequestsVideoCapture())
+	{
+		if (IVideoRecordingSystem* const VideoRecordingSystem = IPlatformFeaturesModule::Get().GetVideoRecordingSystem())
+		{
+			VideoRecordingSystem->GetOnVideoRecordingFinalizedDelegate().RemoveAll(this);
+		}
+	}
+}
diff --git a/Engine/Plugins/Performance/AutomatedPerfTesting/Source/AutomatedPerfTesting/Private/AutomatedPerfTestGameModeBase.cpp b/Engine/Plugins/Performance/AutomatedPerfTesting/Source/AutomatedPerfTesting/Private/AutomatedPerfTestGameModeBase.cpp
new file mode 100644
--- /dev/null
+++ b/Engine/Plugins/Performance/AutomatedPerfTesting/Source/AutomatedPerfTesting/Private/AutomatedPerfTestGameModeBase.cpp
@@ -0,0 +1,24 @@
+﻿// Copyright Epic Games, Inc. All Rights Reserved.
+
+#include "AutomatedPerfTestGameModeBase.h"
+#include "AutomatedPerfTesting.h"
+
+void AAutomatedPerfTestGameModeBase::SetupTest_Implementation()
+{
+	UE_LOG(LogAutomatedPerfTest, Log, TEXT("AutomatedPerfTestGameModeBase SetupTest"))
+}
+
+void AAutomatedPerfTestGameModeBase::TeardownTest_Implementation()
+{
+	UE_LOG(LogAutomatedPerfTest, Log, TEXT("AutomatedPerfTestGameModeBase TeardownTest"))
+}
+
+void AAutomatedPerfTestGameModeBase::RunTest_Implementation()
+{
+	UE_LOG(LogAutomatedPerfTest, Log, TEXT("AutomatedPerfTestGameModeBase RunTest"))
+}
+
+void AAutomatedPerfTestGameModeBase::Exit_Implementation()
+{
+	UE_LOG(LogAutomatedPerfTest, Log, TEXT("AutomatedPerfTestGameModeBase Exit"))
+}
\ No newline at end of file
diff --git a/Engine/Plugins/Performance/AutomatedPerfTesting/Source/AutomatedPerfTesting/Private/AutomatedPerfTestInterface.cpp b/Engine/Plugins/Performance/AutomatedPerfTesting/Source/AutomatedPerfTesting/Private/AutomatedPerfTestInterface.cpp
new file mode 100644
--- /dev/null
+++ b/Engine/Plugins/Performance/AutomatedPerfTesting/Source/AutomatedPerfTesting/Private/AutomatedPerfTestInterface.cpp
@@ -0,0 +1,5 @@
+﻿// Copyright Epic Games, Inc. All Rights Reserved.
+
+#include "AutomatedPerfTestInterface.h"
+
+// Add default functionality here for any IAutomatedPerfTestInterface functions that are not pure virtual.
diff --git a/Engine/Plugins/Performance/AutomatedPerfTesting/Source/AutomatedPerfTesting/Private/AutomatedPerfTestSubsystem.cpp b/Engine/Plugins/Performance/AutomatedPerfTesting/Source/AutomatedPerfTesting/Private/AutomatedPerfTestSubsystem.cpp
new file mode 100644
--- /dev/null
+++ b/Engine/Plugins/Performance/AutomatedPerfTesting/Source/AutomatedPerfTesting/Private/AutomatedPerfTestSubsystem.cpp
@@ -0,0 +1,19 @@
+﻿// Copyright Epic Games, Inc. All Rights Reserved.
+
+#include "AutomatedPerfTestSubsystem.h"
+
+#include "AutomatedPerfTestControllerBase.h"
+#include "GauntletModule.h"
+#include "Modules/ModuleManager.h"
+
+FString UAutomatedPerfTestSubsystem::GetTestID()
+{
+	if(FGauntletModule* ParentModule = &FModuleManager::Get().GetModuleChecked<FGauntletModule>(TEXT("Gauntlet")))
+	{
+		if(UAutomatedPerfTestControllerBase* Controller = ParentModule->GetTestController<UAutomatedPerfTestControllerBase>())
+		{
+			return Controller->GetTestID();
+		}
+	}
+	return TEXT("");
+}
diff --git a/Engine/Plugins/Performance/AutomatedPerfTesting/Source/AutomatedPerfTesting/Private/AutomatedPerfTesting.cpp b/Engine/Plugins/Performance/AutomatedPerfTesting/Source/AutomatedPerfTesting/Private/AutomatedPerfTesting.cpp
new file mode 100644
--- /dev/null
+++ b/Engine/Plugins/Performance/AutomatedPerfTesting/Source/AutomatedPerfTesting/Private/AutomatedPerfTesting.cpp
@@ -0,0 +1,20 @@
+// Copyright Epic Games, Inc. All Rights Reserved.
+
+#include "AutomatedPerfTesting.h"
+
+#define LOCTEXT_NAMESPACE "FAutomatedPerfTestingModule"
+
+void FAutomatedPerfTestingModule::StartupModule()
+{
+	// This code will execute after your module is loaded into memory; the exact timing is specified in the .uplugin file per-module
+}
+
+void FAutomatedPerfTestingModule::ShutdownModule()
+{
+	// This function may be called during shutdown to clean up your module.  For modules that support dynamic reloading,
+	// we call this function before unloading the module.
+}
+
+#undef LOCTEXT_NAMESPACE
+	
+IMPLEMENT_MODULE(FAutomatedPerfTestingModule, AutomatedPerfTesting)
\ No newline at end of file
diff --git a/Engine/Plugins/Performance/AutomatedPerfTesting/Source/AutomatedPerfTesting/Private/AutomatedSequencePerfTest.cpp b/Engine/Plugins/Performance/AutomatedPerfTesting/Source/AutomatedPerfTesting/Private/AutomatedSequencePerfTest.cpp
new file mode 100644
--- /dev/null
+++ b/Engine/Plugins/Performance/AutomatedPerfTesting/Source/AutomatedPerfTesting/Private/AutomatedSequencePerfTest.cpp
@@ -0,0 +1,141 @@
+// Copyright Epic Games, Inc. All Rights Reserved.
+
+
+#include "AutomatedSequencePerfTest.h"
+
+#include "TimerManager.h"
+#include "AutomatedPerfTesting.h"
+#include "LevelSequenceActor.h"
+#include "LevelSequencePlayer.h"
+#include "MovieSceneSequence.h"
+#include "MovieSceneSequencePlayer.h"
+#include "Misc/CommandLine.h"
+
+namespace AutomatedPerfTesting
+{
+	static TAutoConsoleVariable<FString> CvarSequencePath(
+	TEXT("AutomatedPerfTest.SequencePath"),
+	"",
+	TEXT("Full path to the sequence which should be run when the AutomatedSequencePerfTest is executed. EG. /Game/Profiling/SampleSequence.SampleSequence"),
+	ECVF_Default);
+}
+
+FString UAutomatedSequencePerfTest::GetSequenceName() const
+{
+	return SequenceSoftPath.GetAssetName();
+}
+
+void UAutomatedSequencePerfTest::SetupTest()
+{
+	Super::SetupTest();
+
+	UE_LOG(LogAutomatedPerfTest, Log, TEXT("Setting up test"));
+
+	// make sure the world exists, then create a sequence player
+	if(UWorld* const World = GetWorld())
+	{
+		// load the sequence specified by the user
+		UE_LOG(LogAutomatedPerfTest, Log, TEXT("Loading sequence %s"), *SequenceSoftPath.ToString());
+		ULevelSequence* TargetSequence = LoadObject<ULevelSequence>(NULL, *SequenceSoftPath.ToString(), NULL, LOAD_None, NULL);
+		check(TargetSequence);
+		
+		UE_LOG(LogAutomatedPerfTest, Log, TEXT("World is valid, creating sequence player"));
+		SequencePlayer = ULevelSequencePlayer::CreateLevelSequencePlayer(World, TargetSequence, FMovieSceneSequencePlaybackSettings(), SequenceActor);
+
+		if (SequencePlayer == nullptr)
+		{
+			UE_LOG(LogAutomatedPerfTest, Error, TEXT("Unable to create sequence player when starting AutomatedSequencePerfTest, exiting..."));
+			EndAutomatedPerfTest(1);
+		}
+
+		// set the sequence up at the beginning
+		FMovieSceneSequencePlaybackParams PlaybackParams = FMovieSceneSequencePlaybackParams();
+		FMovieSceneSequencePlayToParams PlayToParams = FMovieSceneSequencePlayToParams();
+
+		PlaybackParams.Time = 0.0;
+		PlaybackParams.UpdateMethod = EUpdatePositionMethod::Scrub;
+
+		UE_LOG(LogAutomatedPerfTest, Log, TEXT("RunTest:: Scrubbing to start"));
+		SequencePlayer->PlayTo(PlaybackParams, PlayToParams);
+		
+		FTimerHandle UnusedHandle;
+		
+		// TODO parameterize the presoak delay
+		GetWorld()->GetTimerManager().SetTimer(UnusedHandle, this, &UAutomatedSequencePerfTest::RunTest, 1.0, false, 5.0);
+	}
+	// if we have an invalid world, we can't run the test, so we should bail out
+	else
+	{
+		UE_LOG(LogAutomatedPerfTest, Error, TEXT("Invalid World when starting AutomatedSequencePerfTest, exiting..."));
+		EndAutomatedPerfTest(1);
+	}
+}
+
+void UAutomatedSequencePerfTest::RunTest()
+{
+	Super::RunTest();
+
+	UE_LOG(LogAutomatedPerfTest, Log, TEXT("RunTest"));
+	
+	// make sure we have a valid sequence player
+	if(SequencePlayer)
+	{
+		UE_LOG(LogAutomatedPerfTest, Log, TEXT("RunTest::Valid Sequence Player, proceeding"));
+		
+        SequencePlayer->Play();
+		// TODO add a bit here that also triggers begin/end regions on camera cut
+        SequencePlayer->OnFinished.AddDynamic(this, &UAutomatedSequencePerfTest::TeardownTest);
+	}
+	// otherwise bail out of the test
+	else
+	{
+		UE_LOG(LogAutomatedPerfTest, Error, TEXT("Invalid SequencePlayer when starting AutomatedSequencePerfTest, exiting..."));
+		EndAutomatedPerfTest(1);
+	}
+}
+
+void UAutomatedSequencePerfTest::TeardownTest()
+{
+	UE_LOG(LogAutomatedPerfTest, Log, TEXT("AutomatedSequencePerfTest::TeardownTest"));
+	
+	Super::TeardownTest();
+}
+
+void UAutomatedSequencePerfTest::Exit()
+{
+	UE_LOG(LogAutomatedPerfTest, Log, TEXT("AutomatedSequencePerfTest::Exit"));
+	Super::Exit();
+}
+
+void UAutomatedSequencePerfTest::OnInit()
+{
+	Super::OnInit();
+	UE_LOG(LogAutomatedPerfTest, Log, TEXT("AutomatedSequencePerfTest::OnInit"));
+	
+	if (FParse::Value(FCommandLine::Get(), TEXT("AutomatedPerfTest.SequencePath="), SequencePathName))
+	{
+		SequencePathName.ReplaceInline(TEXT("\""), TEXT(""));
+		SequencePathName.ReplaceInline(TEXT("="), TEXT(""));
+		AutomatedPerfTesting::CvarSequencePath->Set(*SequencePathName);
+		UE_LOG(LogAutomatedPerfTest, Log, TEXT("Input Sequence Path Name = %s"), *SequencePathName);
+		SequenceSoftPath = FSoftObjectPath(SequencePathName);
+	}
+}
+
+void UAutomatedSequencePerfTest::UnbindAllDelegates()
+{
+	Super::UnbindAllDelegates();
+
+	// if we have a valid sequence player, make sure we unbind our events from it when we're wrapping up the test.
+	if(SequencePlayer != nullptr)
+	{
+		SequencePlayer->OnFinished.RemoveAll(this);
+	}
+
+	// clear any stray timers that might be laying around
+	GetWorld()->GetTimerManager().ClearAllTimersForObject(this);
+	if(SequencePlayer)
+	{
+		GetWorld()->GetTimerManager().ClearAllTimersForObject(SequencePlayer);
+	}
+}
diff --git a/Engine/Plugins/Performance/AutomatedPerfTesting/Source/AutomatedPerfTesting/Private/AutomatedSoakTest.cpp b/Engine/Plugins/Performance/AutomatedPerfTesting/Source/AutomatedPerfTesting/Private/AutomatedSoakTest.cpp
new file mode 100644
--- /dev/null
+++ b/Engine/Plugins/Performance/AutomatedPerfTesting/Source/AutomatedPerfTesting/Private/AutomatedSoakTest.cpp
@@ -0,0 +1,3 @@
+﻿// Copyright Epic Games, Inc. All Rights Reserved.
+
+#include "AutomatedSoakTest.h"
diff --git a/Engine/Plugins/Performance/AutomatedPerfTesting/Source/AutomatedPerfTesting/Public/AutomatedPerfTestControllerBase.h b/Engine/Plugins/Performance/AutomatedPerfTesting/Source/AutomatedPerfTesting/Public/AutomatedPerfTestControllerBase.h
new file mode 100644
--- /dev/null
+++ b/Engine/Plugins/Performance/AutomatedPerfTesting/Source/AutomatedPerfTesting/Public/AutomatedPerfTestControllerBase.h
@@ -0,0 +1,101 @@
+﻿// Copyright Epic Games, Inc. All Rights Reserved.
+
+#pragma once
+
+#include "CoreMinimal.h"
+#include "GauntletTestController.h"
+#include "Engine/World.h"
+
+#include "AutomatedPerfTestControllerBase.generated.h"
+
+class AGameModeBase;
+
+/**
+ * 
+ */
+UCLASS()
+class AUTOMATEDPERFTESTING_API UAutomatedPerfTestControllerBase : public UGauntletTestController
+{
+	GENERATED_BODY()
+public:
+	void OnPreWorldInitializeInternal(UWorld* World, const UWorld::InitializationValues IVS);
+	virtual void OnPreWorldInitialize(UWorld* World);
+
+	UFUNCTION()
+	void TryEarlyExec(UWorld* const World);
+	
+	UFUNCTION()
+	void OnWorldBeginPlay();
+
+	UFUNCTION()
+	void OnGameStateSet(AGameStateBase* const GameStateBase);
+
+// Base functionality
+public:
+	UAutomatedPerfTestControllerBase(const FObjectInitializer& ObjectInitializer);
+
+	static FString GetTestName();
+	static FString GetDeviceProfile();
+	static FString GetExplicitTestID();
+	FString GetTestID();
+	FString GetOverallRegionName();
+	
+	bool RequestsInsightsTrace() const;
+	bool RequestsCSVProfiler() const;
+	bool RequestsFPSChart() const;
+	bool RequestsVideoCapture() const;
+
+	bool TryStartInsightsTrace();
+	bool TryStopInsightsTrace();
+
+	bool TryStartCSVProfiler();
+	bool TryStopCSVProfiler();
+	
+	bool TryStartFPSChart();
+	bool TryStopFPSChart();
+
+	bool TryStartVideoCapture();
+	bool TryFinalizingVideoCapture(const bool bStopAutoContinue = false);
+
+	virtual void SetupTest();
+	virtual void RunTest();
+	virtual void TeardownTest();
+	virtual void Exit();
+	
+protected:
+	// ~Begin UGauntletTestController Interface
+	virtual void OnInit() override;
+	virtual void OnTick(float TimeDelta) override;
+	virtual void OnStateChange(FName OldState, FName NewState) override;
+	virtual void OnPreMapChange() override;
+	virtual void BeginDestroy() override;
+	// ~End UGauntletTestController Interface
+
+	UFUNCTION()
+	virtual void EndAutomatedPerfTest(const int32 ExitCode = 0);
+
+	UFUNCTION()
+	virtual void OnVideoRecordingFinalized(bool Succeeded, const FString& FilePath);
+	
+	virtual void UnbindAllDelegates();
+
+private:
+	FString ExplicitTestID;
+
+	FString TraceChannels;
+	FString TestDatetime;
+	FString TestID;
+	FString DeviceProfileOverride;
+	bool bRequestsFPSChart;
+	bool bRequestsInsightsTrace;
+	bool bRequestsCSVProfiler;
+	bool bRequestsVideoCapture;
+
+	FText VideoRecordingTitle;
+	
+	const TArray<FString> CmdsToExecEarly = { };
+	
+	AGameModeBase* GameMode;
+
+	FDelegateHandle CsvProfilerDelegateHandle;
+};
diff --git a/Engine/Plugins/Performance/AutomatedPerfTesting/Source/AutomatedPerfTesting/Public/AutomatedPerfTestGameModeBase.h b/Engine/Plugins/Performance/AutomatedPerfTesting/Source/AutomatedPerfTesting/Public/AutomatedPerfTestGameModeBase.h
new file mode 100644
--- /dev/null
+++ b/Engine/Plugins/Performance/AutomatedPerfTesting/Source/AutomatedPerfTesting/Public/AutomatedPerfTestGameModeBase.h
@@ -0,0 +1,35 @@
+﻿// Copyright Epic Games, Inc. All Rights Reserved.
+
+#pragma once
+
+#include "CoreMinimal.h"
+#include "AutomatedPerfTestInterface.h"
+#include "GameFramework/GameModeBase.h"
+#include "AutomatedPerfTestGameModeBase.generated.h"
+
+/**
+ * 
+ */
+UCLASS(Blueprintable)
+class AUTOMATEDPERFTESTING_API AAutomatedPerfTestGameModeBase : public AGameModeBase, public IAutomatedPerfTestInterface
+{
+		GENERATED_BODY()
+	
+public:
+
+//~ Begin IAutomatedPerfTestInterface	
+	UFUNCTION(BlueprintCallable, BlueprintImplementableEvent, Category="Automated Perf Test")
+	void SetupTest(); virtual void SetupTest_Implementation();
+
+	UFUNCTION(BlueprintCallable, BlueprintImplementableEvent, Category="Automated Perf Test")
+	void TeardownTest(); virtual void TeardownTest_Implementation();
+
+	UFUNCTION(BlueprintCallable, BlueprintImplementableEvent, Category="Automated Perf Test")
+	void RunTest(); virtual void RunTest_Implementation();
+
+	UFUNCTION(BlueprintCallable, BlueprintImplementableEvent, Category="Automated Perf Test")
+	void Exit(); virtual void Exit_Implementation();
+//~ End IAutomatedPerfTestInterface
+
+	
+};
diff --git a/Engine/Plugins/Performance/AutomatedPerfTesting/Source/AutomatedPerfTesting/Public/AutomatedPerfTestInterface.h b/Engine/Plugins/Performance/AutomatedPerfTesting/Source/AutomatedPerfTesting/Public/AutomatedPerfTestInterface.h
new file mode 100644
--- /dev/null
+++ b/Engine/Plugins/Performance/AutomatedPerfTesting/Source/AutomatedPerfTesting/Public/AutomatedPerfTestInterface.h
@@ -0,0 +1,35 @@
+﻿// Copyright Epic Games, Inc. All Rights Reserved.
+#pragma once
+
+#include "CoreMinimal.h"
+#include "UObject/Interface.h"
+#include "AutomatedPerfTestInterface.generated.h"
+
+// This class does not need to be modified.
+UINTERFACE(BlueprintType)
+class AUTOMATEDPERFTESTING_API UAutomatedPerfTestInterface : public UInterface
+{
+	GENERATED_BODY()
+};
+
+/**
+ * 
+ */
+class AUTOMATEDPERFTESTING_API IAutomatedPerfTestInterface
+{
+	GENERATED_BODY()
+
+	// Add interface functions to this class. This is the class that will be inherited to implement this interface.
+public:
+	UFUNCTION(BlueprintCallable, BlueprintImplementableEvent, Category="Automated Perf Test")
+	void SetupTest();
+
+	UFUNCTION(BlueprintCallable, BlueprintImplementableEvent, Category="Automated Perf Test")
+	void TeardownTest();
+
+	UFUNCTION(BlueprintCallable, BlueprintImplementableEvent, Category="Automated Perf Test")
+	void RunTest();
+
+	UFUNCTION(BlueprintCallable, BlueprintImplementableEvent, Category="Automated Perf Test")
+	void Exit();
+};
diff --git a/Engine/Plugins/Performance/AutomatedPerfTesting/Source/AutomatedPerfTesting/Public/AutomatedPerfTestSubsystem.h b/Engine/Plugins/Performance/AutomatedPerfTesting/Source/AutomatedPerfTesting/Public/AutomatedPerfTestSubsystem.h
new file mode 100644
--- /dev/null
+++ b/Engine/Plugins/Performance/AutomatedPerfTesting/Source/AutomatedPerfTesting/Public/AutomatedPerfTestSubsystem.h
@@ -0,0 +1,20 @@
+﻿// Copyright Epic Games, Inc. All Rights Reserved.
+
+#pragma once
+
+#include "CoreMinimal.h"
+#include "Subsystems/EngineSubsystem.h"
+#include "AutomatedPerfTestSubsystem.generated.h"
+
+/**
+ * 
+ */
+UCLASS()
+class AUTOMATEDPERFTESTING_API UAutomatedPerfTestSubsystem : public UEngineSubsystem
+{
+	GENERATED_BODY()
+
+	UFUNCTION(BlueprintCallable, Category="Automated Perf Testing")
+	static FString GetTestID();
+	
+};
diff --git a/Engine/Plugins/Performance/AutomatedPerfTesting/Source/AutomatedPerfTesting/Public/AutomatedPerfTesting.h b/Engine/Plugins/Performance/AutomatedPerfTesting/Source/AutomatedPerfTesting/Public/AutomatedPerfTesting.h
new file mode 100644
--- /dev/null
+++ b/Engine/Plugins/Performance/AutomatedPerfTesting/Source/AutomatedPerfTesting/Public/AutomatedPerfTesting.h
@@ -0,0 +1,19 @@
+// Copyright Epic Games, Inc. All Rights Reserved.
+
+#pragma once
+
+#include "CoreMinimal.h"
+#include "Modules/ModuleManager.h"
+#include "ProfilingDebugging/CsvProfiler.h"
+
+DECLARE_LOG_CATEGORY_EXTERN(LogAutomatedPerfTest, Log, All)
+CSV_DECLARE_CATEGORY_EXTERN(AutomatedPerfTest);
+
+class FAutomatedPerfTestingModule : public IModuleInterface
+{
+public:
+
+	/** IModuleInterface implementation */
+	virtual void StartupModule() override;
+	virtual void ShutdownModule() override;
+};
diff --git a/Engine/Plugins/Performance/AutomatedPerfTesting/Source/AutomatedPerfTesting/Public/AutomatedSequencePerfTest.h b/Engine/Plugins/Performance/AutomatedPerfTesting/Source/AutomatedPerfTesting/Public/AutomatedSequencePerfTest.h
new file mode 100644
--- /dev/null
+++ b/Engine/Plugins/Performance/AutomatedPerfTesting/Source/AutomatedPerfTesting/Public/AutomatedSequencePerfTest.h
@@ -0,0 +1,43 @@
+﻿// Copyright Epic Games, Inc. All Rights Reserved.
+
+#pragma once
+
+#include "CoreMinimal.h"
+#include "AutomatedPerfTestControllerBase.h"
+#include "AutomatedSequencePerfTest.generated.h"
+
+class ULevelSequence;
+class ALevelSequenceActor;
+class ULevelSequencePlayer;
+
+/**
+ * 
+ */
+UCLASS()
+class AUTOMATEDPERFTESTING_API UAutomatedSequencePerfTest : public UAutomatedPerfTestControllerBase
+{
+	GENERATED_BODY()
+
+public:
+	FString GetSequenceName() const;
+	
+	virtual void SetupTest() override;
+
+	UFUNCTION()
+	virtual void RunTest() override;
+
+	UFUNCTION()
+	virtual void TeardownTest() override;
+	virtual void Exit() override;
+	
+protected:
+	virtual void OnInit() override;
+	virtual void UnbindAllDelegates() override;
+
+private:
+	FString SequencePathName;
+	FSoftObjectPath SequenceSoftPath;
+
+	ALevelSequenceActor* SequenceActor;
+	ULevelSequencePlayer* SequencePlayer;
+};
diff --git a/Engine/Plugins/Performance/AutomatedPerfTesting/Source/AutomatedPerfTesting/Public/AutomatedSoakTest.h b/Engine/Plugins/Performance/AutomatedPerfTesting/Source/AutomatedPerfTesting/Public/AutomatedSoakTest.h
new file mode 100644
--- /dev/null
+++ b/Engine/Plugins/Performance/AutomatedPerfTesting/Source/AutomatedPerfTesting/Public/AutomatedSoakTest.h
@@ -0,0 +1,16 @@
+﻿// Copyright Epic Games, Inc. All Rights Reserved.
+
+#pragma once
+
+#include "CoreMinimal.h"
+#include "AutomatedPerfTestControllerBase.h"
+#include "AutomatedSoakTest.generated.h"
+
+/**
+ * 
+ */
+UCLASS()
+class AUTOMATEDPERFTESTING_API UAutomatedSoakTest : public UAutomatedPerfTestControllerBase
+{
+	GENERATED_BODY()
+};
