Automated Performance Testing - Adding support for multiple map/sequence combinations
+ Adding BuildGraph for common APT options
+ Adding BuildGraph for stubbing in AutomatedSequencerPerfTest for AdditionalTestNodes
- Removing extraneous commandline and config elements from AutomatedPerfTestConfig
+ Adding first pass at AutomatedPerfTest Insights reporting build command
+ Adding APT Project Settings
* Updating projects to leverage the new multiple map/sequence combo feature of APT

diff --git a/Engine/Plugins/Performance/AutomatedPerfTesting/Build/Inc/AutomatedPerfTestCommonSettings.xml b/Engine/Plugins/Performance/AutomatedPerfTesting/Build/Inc/AutomatedPerfTestCommonSettings.xml
new file mode 100644
--- /dev/null
+++ b/Engine/Plugins/Performance/AutomatedPerfTesting/Build/Inc/AutomatedPerfTestCommonSettings.xml
@@ -0,0 +1,41 @@
+<?xml version="1.0" encoding="utf-8"?>
+<BuildGraph xmlns="http://www.epicgames.com/BuildGraph" xmlns:xsi="http://www.w3.org/2001/XMLSchema-instance" xsi:schemaLocation="http://www.epicgames.com/BuildGraph ../../../../../../../Engine/Build/Graph/Schema.xsd" >
+
+	<!-- Insights Setup for Automated Perf Testing -->
+	<!-- since we need to tell UAT whether or not to using the shipping configuration of Insights when processing -->
+	<!-- insights traces after the conclusion of the test --> 
+	<Option Name="UseShippingInsights" Restrict="true|false" DefaultValue="false" Description="Whether or not to use the shipping configuration of Insights." />
+
+	<!-- then we'll use that to override values that will be picked up by the rest of BuildAndTestProject.xml -->
+	<Property Name="IncludeUnrealInsightsInToolCompile" Value="true" />
+	<Property Name="UnrealInsightsBuildConfiguration" Value="Development" If="!$(UseShippingInsights)" />
+	<Property Name="UnrealInsightsBuildConfiguration" Value="Shipping" If="$(UseShippingInsights)" />
+	
+	<!-- Other common arguments --> 
+	<Option Name="TestID" DefaultValue="" Description="A string to use to identify these performance tests" />
+	<Option Name="DeviceProfileOverride" DefaultValue="" Description="Which device profile to use for this test instead of the default value. Will be prepended by PlatformName." />
+	<Option Name="DoInsightsTrace" DefaultValue="false" Restrict="true|false" Description="Whether or not to use Unreal Insights to trace and parse performance results." />
+	<Option Name="DoCSVProfiler" DefaultValue="true" Restrict="true|false" Description="Whether or not to use CSVProfiler to gather and parse performance telemetry." />
+	<Option Name="PerfReportServer" DefaultValue="false" Restrict="true|false" Description="Whether or not to use CSVProfiler to gather and parse performance telemetry." />
+	<Option Name="DoFPSChart" DefaultValue="false" Restrict="true|false" Description="Whether or not to use FPS Chart to gather and parse performance telemetry." />
+	<Option Name="TraceChannels" DefaultValue="default,screenshot,stats" Description="Which channels to use when tracing with Insights" />
+	<Option Name="DoVideoCapture" Restrict="true|false" DefaultValue="false" Description="Whether or not to capture video during the test" />
+	<Option Name="DataSourceName" DefaultValue="Automation.$(ProjectName).Performance" Description="To which PRS data source should this be uploaded"/>
+	<Option Name="PerfCacheRoot" DefaultValue="$(RootDir)/$(ProjectPath)/PerfCache" Description="Path to the perf cache where all the output data is stored. Defaults to a path in the local project's directory. Override to a common network-accessible path when running on build machines"/>
+	<Option Name="GenerateLocalReports" DefaultValue="false" Restrict="true|false" Description="Whether or not to create a local PerfReportTool report" />
+	
+	<!-- Construct the common arguments for the performance test and report commands -->
+	<Property Name="AutomatedPerformanceCommonArgs" Value = "-AutomatedPerfTest.UseShippingInsights=$(UseShippingInsights)" />
+	<Property Name="AutomatedPerformanceCommonArgs" Value="$(AutomatedPerformanceCommonArgs) -AutomatedPerfTest.TestName=$(TestID)"/>
+	<Property Name="AutomatedPerformanceCommonArgs" Value="$(AutomatedPerformanceCommonArgs) -AutomatedPerfTest.DeviceProfileOverride=$(DeviceProfileOverride)" If="'$(DeviceProfileOverride)' != ''"/>
+	<Property Name="AutomatedPerformanceCommonArgs" Value="$(AutomatedPerformanceCommonArgs) -AutomatedPerfTest.DoInsightsTrace" If="$(DoInsightsTrace)"/>
+	<Property Name="AutomatedPerformanceCommonArgs" Value="$(AutomatedPerformanceCommonArgs) -AutomatedPerfTest.TraceChannels=$(TraceChannels)" If="$(DoInsightsTrace) and '$(TraceChannels)' != ''"/>
+	<Property Name="AutomatedPerformanceCommonArgs" Value="$(AutomatedPerformanceCommonArgs) -AutomatedPerfTest.DoCSVProfiler" If="$(DoCSVProfiler)"/>
+	<Property Name="AutomatedPerformanceCommonArgs" Value="$(AutomatedPerformanceCommonArgs) -AutomatedPerfTest.DoFPSChart" If="$(DoFPSChart)"/>
+	<Property Name="AutomatedPerformanceCommonArgs" Value="$(AutomatedPerformanceCommonArgs) -AutomatedPerfTest.PerfCacheRoot=$(PerfCacheRoot)"/>
+	<Property Name="AutomatedPerformanceCommonArgs" Value="$(AutomatedPerformanceCommonArgs) -PerfReportServer" If="$(PerfReportServer) And $(DoCSVProfiler)"/>
+	<Property Name="AutomatedPerformanceCommonArgs" Value="$(AutomatedPerformanceCommonArgs) -NoLocalReports" If="!$(GenerateLocalReports) And $(DoCSVProfiler)" />
+	<Property Name="AutomatedPerformanceCommonArgs" Value="$(AutomatedPerformanceCommonArgs) -resX=3840 -resY=2160"/>
+	<Property Name="AutomatedPerformanceCommonArgs" Value="$(AutomatedPerformanceCommonArgs) -AutomatedPerfTest.DoVideoCapture" If="$(DoVideoCapture)"/>
+
+</BuildGraph>
diff --git a/Engine/Plugins/Performance/AutomatedPerfTesting/Build/Inc/AutomatedPerfTestInsightsReport.xml b/Engine/Plugins/Performance/AutomatedPerfTesting/Build/Inc/AutomatedPerfTestInsightsReport.xml
new file mode 100644
--- /dev/null
+++ b/Engine/Plugins/Performance/AutomatedPerfTesting/Build/Inc/AutomatedPerfTestInsightsReport.xml
@@ -0,0 +1,26 @@
+﻿<?xml version='1.0' ?>
+<BuildGraph xmlns="http://www.epicgames.com/BuildGraph" xmlns:xsi="http://www.w3.org/2001/XMLSchema-instance" xsi:schemaLocation="http://www.epicgames.com/BuildGraph ../../../../../../../Engine/Build/Graph/Schema.xsd" >
+	<!-- =============================================================================== -->
+	<!-- Include this graph in your project's Tests.xml graph AFTER BuildAndTestProject.xml in order to 
+	     process the Insights traces that were generated by all the previous tests       
+	     You also need to make sure that somewhere before this is included you've also included
+	     AutomatedPerfTestCommonSettings.xml                                             -->
+	<!-- =============================================================================== -->
+	
+	<!-- If an Insights trace was requested, run the report job on a TestWin64 and have it process all of the unprocessed traces -->
+	<Do If="$(DoInsightsTrace)">
+		<!-- since this command is wholly constructed here, we need to send in things like Project and Build names --> 
+		<Property Name="InsightsReportCommonArguments" Value="$(AutomatedPerformanceCommonArgs)" />
+		<Property Name="InsightsReportCommonArguments" Value="$(InsightsReportCommonArguments) AutomatedPerfTest.BuildName=$(BuildNamePath)" />
+		<Property Name="InsightsReportCommonArguments" Value="$(InsightsReportCommonArguments) AutomatedPerfTest.ProjectName=$(ProjectName)" />
+		
+		<Property Name="InsightsReportNodeNameBase" Value="$(ProjectName) Insights Report" />
+		<Agent Name="$(InsightsReportNodeNameBase)" Type="TestWin64">
+			<Node Name="$(ProjectName) Insights Report" Requires="BuildAndTest $(ProjectName)">
+				<Command Name="$(ProjectName)TestInsightsReport" Arguments="$(AutomatedPerformanceCommonArgs)" />
+			</Node>
+		</Agent>
+		<Property Name="BuildAndTestExtendedRequirements" Value="$(BuildAndTestExtendedRequirements);$(InsightsReportNodeNameBase)" />
+	</Do>
+
+</BuildGraph>
diff --git a/Engine/Plugins/Performance/AutomatedPerfTesting/Build/Inc/AutomatedSequencePerfTestMacro.xml b/Engine/Plugins/Performance/AutomatedPerfTesting/Build/Inc/AutomatedSequencePerfTestMacro.xml
new file mode 100644
--- /dev/null
+++ b/Engine/Plugins/Performance/AutomatedPerfTesting/Build/Inc/AutomatedSequencePerfTestMacro.xml
@@ -0,0 +1,38 @@
+<?xml version='1.0' ?>
+<BuildGraph xmlns="http://www.epicgames.com/BuildGraph" xmlns:xsi="http://www.w3.org/2001/XMLSchema-instance" xsi:schemaLocation="http://www.epicgames.com/BuildGraph ../../../../../../../Engine/Build/Graph/Schema.xsd" >
+
+	<Macro Name="AutomatedSequencePerfTestNode">
+		<!-- Run the test. No 'Requires' as we run on the same machine locally and use the publish folder otherwise -->
+		<Node Name="$(TestNodeName)" Requires="$(TestPrerequisites)" Annotations="$(NodeAnnotations)">
+			<Command Name="RunUnreal" Arguments="-project=$(TargetProject) -test=&quot;$(TestName)&quot; -platform=$(TargetPlatform) -configuration=$(TestConfiguration) -target=$(TargetName) $(GauntletBuildArg) $(GauntletTestArgs) $(GauntletReportArgs) $(AdditionalArguments)"/>
+		</Node>
+	</Macro>
+	
+	<Macro Name="AutomatedSequencePerfTest" Arguments="TestName;MapSequenceComboNames;AdditionalArguments">
+		<ForEach Name="MapSequenceName" Values="$(MapSequenceComboNames)" Separator="+">
+			<Property Name="TestNameTitle" Value="$(TestName)"/>
+			<Expand Name="SanitizeTestNameTitle"/>
+			
+			<!-- Run and report on the test-->
+			<Do If="'$(AutomationReportOutputDirectory)' != ''">
+				<Property Name="OutputFragmentPath" Value="$(TargetPlatform)/$(TestConfiguration)/$(TestNameTitle)" />
+				<Property Name="OutputReportPath" Value="$(AutomationReportOutputDirectory)/$(OutputFragmentPath)" />
+				<Property Name="GauntletReportArgs" Value="-logdir=&quot;$(OutputReportPath)&quot; -artifactname=Gauntlet -ReportExportPath=&quot;$(OutputReportPath)&quot;" />
+				<Property Name="GauntletReportArgs" Value="$(GauntletReportArgs) -ReportURL=&quot;$(AutomationReportUrl)/$(OutputFragmentPath)/&quot;" If="'$(AutomationReportUrl)' != ''" />
+			</Do>
+			<Property Name="TestNodeName" Value="$(PreNodeName)$(TargetPlatform) $(TestConfiguration) $(TestNameTitle) $(MapSequenceName)" />
+			<Property Name="NodeAnnotations" Value=""/>
+			<Property Name="NodeAnnotations" Value="Workflow=$(AnnotationsTarget)" If="ContainsItem('$(AnnotationsPlatforms)', '$(TargetPlatform)', ';') and ContainsItem('$(AnnotationsTests)', '$(TestName)', ';')"/>
+			
+			<Property Name="GauntletTestArgs" Value="$(GauntletTestArgs) -AutomatedPerfTest.SequencePerfTest.MapSequenceName=$(MapSequenceName)" />
+			
+			<Expand Name="AutomatedSequencePerfTestNode" />
+
+			<!-- Add it to our list -->
+			<Do If="!ContainsItem('$(TargetPlatformsToSkipTests)', '$(TargetPlatform)', '+')">
+				<Property Name="AllPlatformTestNodes" Value="$(AllPlatformTestNodes);$(TestNodeName)"/>
+			</Do>
+		</ForEach>
+	</Macro>
+	
+</BuildGraph>
diff --git a/Engine/Plugins/Performance/AutomatedPerfTesting/Build/Scripts/AutomatedPerfTestConfig.cs b/Engine/Plugins/Performance/AutomatedPerfTesting/Build/Scripts/AutomatedPerfTestConfig.cs
--- a/Engine/Plugins/Performance/AutomatedPerfTesting/Build/Scripts/AutomatedPerfTestConfig.cs
+++ b/Engine/Plugins/Performance/AutomatedPerfTesting/Build/Scripts/AutomatedPerfTestConfig.cs
@@ -11,24 +11,12 @@ namespace AutomatedPerfTest
 {
 	public class AutomatedPerfTestConfigBase : UnrealTestConfiguration
 	{
-		/// <summary>
-		/// Name of the project (TODO: parse from the build metadata? What about github builds?)
-		/// </summary>
-		[AutoParamWithNames("AutomatedPerfTest.ProjectName")]
-		public string ProjectName = "";
-
 		/// <summary>
 		/// Name of the project (TODO: parse from the build metadata? What about github builds?)
 		/// </summary>
 		[AutoParamWithNames("AutomatedPerfTest.DataSourceName")]
 		public string DataSourceName = "";
 
-		/// <summary>
-		/// If we're running on the build machine
-		/// </summary>
-		[AutoParamWithNames(false, "AutomatedPerfTest.IsBuildMachine")]
-		public bool IsBuildMachine;
-
 		/// <summary>
 		/// Name of the test, useful for identifying it later
 		/// </summary>
@@ -51,7 +39,7 @@ namespace AutomatedPerfTest
 		/// Which trace channels to test with
 		/// </summary>
 		[AutoParamWithNames("AutomatedPerfTest.TraceChannels")]
-		public string TraceChannels = "default,screenshot";
+		public string TraceChannels = "default,screenshot,stats";
 		
 		/// <summary>
 		/// If we're running on the build machine
@@ -76,18 +64,19 @@ namespace AutomatedPerfTest
 		/// </summary>
 		[AutoParamWithNames("", "AutomatedPerfTest.IgnoredIssuesConfigAbsPath")]
 		public string IgnoredIssuesConfigAbsPath;
-
-		/// <summary>
-		/// Which platform we're testing on 
-		/// </summary>
-		[AutoParamWithNames("", "AutomatedPerfTest.TestPlatform")]
-		public string TestPlatform;
 		
 		/// <summary>
-		/// Let BuildGraph tell us where we should output the Insights trace after running a test so that we know where to grab it from when we're done
+		/// If we should trigger a video capture
 		/// </summary>
 		[AutoParamWithNames(false, "AutomatedPerfTest.DoVideoCapture")]
 		public bool DoVideoCapture;
+		
+		/// <summary>
+		/// If true, will look for the shipping configuration of Unreal Insights in order to parse Insights trace files
+		/// This should be True unless you need to test an issue in parsing the Insights file itself.
+		/// </summary>
+		[AutoParamWithNames(true, "AutomatedPerfTest.UseShippingInsights")]
+		public bool UseShippingInsights;
 	}
 
 	public class AutomatedSequencePerfTestConfig : AutomatedPerfTestConfigBase
@@ -95,13 +84,7 @@ namespace AutomatedPerfTest
 		/// <summary>
 		/// Which map to run the test on
 		/// </summary>
-		[AutoParamWithNames("", "AutomatedPerfTest.TestMap")]
-		public string TestMap;
-		
-		/// <summary>
-		/// Which sequence to run the test with
-		/// </summary>
-		[AutoParamWithNames("", "AutomatedPerfTest.SequencePath")]
-		public string SequencePath;
+		[AutoParamWithNames("", "AutomatedPerfTest.SequencePerfTest.MapSequenceName")]
+		public string MapSequenceComboName;
 	}
 }
diff --git a/Engine/Plugins/Performance/AutomatedPerfTesting/Build/Scripts/AutomatedPerfTestInsightsReport.cs b/Engine/Plugins/Performance/AutomatedPerfTesting/Build/Scripts/AutomatedPerfTestInsightsReport.cs
new file mode 100644
--- /dev/null
+++ b/Engine/Plugins/Performance/AutomatedPerfTesting/Build/Scripts/AutomatedPerfTestInsightsReport.cs
@@ -0,0 +1,336 @@
+using AutomationTool;
+using EpicGames.Core;
+using System;
+using System.Collections.Generic;
+using System.IO;
+using System.Linq;
+using System.Text.Json;
+using UnrealBuildBase;
+using Log = EpicGames.Core.Log;
+using Microsoft.Extensions.Logging;
+
+namespace AutomatedPerfTest
+{
+	/// <summary>
+	/// Helper class for parsing and accessing the metadata stored in a trace's filename
+	/// </summary>
+	public class TraceFileInfo
+	{
+		public string BuildName;
+		public string Build;
+		public int Changelist;
+		public int PreflightChangelist = -1;
+		public string PlatformName;
+		public string DateTime;
+		public string DeviceProfile;
+		public string TestName;
+		public string TraceFilePath;
+		public string FilePath;
+		public string FileName;
+
+		public TraceFileInfo(string InTraceFilePath)
+		{
+			TraceFilePath = InTraceFilePath;
+			FilePath = Path.GetDirectoryName(TraceFilePath);
+			FileName = Path.GetFileNameWithoutExtension(TraceFilePath);
+			
+			// Trace files come out in the form:
+			// Megaworlds-CL-0_Windows_20231120-112520_Windows_InitialTesting.utrace
+			
+			string[] FilenameTokens = FileName.Split('_');
+			BuildName = FilenameTokens[0].ToLower();
+			Build = BuildName.Split('-')[0].ToLower();
+			Changelist = int.Parse(BuildName.Split('-')[2]);
+			
+			// check for a preflight changelist so we can strip out the preflight ID from the filename
+			// due to path-length limitations
+			if(BuildName.Contains("pf"))
+			{
+				// strip the preflight info out of the build name
+				BuildName = String.Join('-', BuildName.Split('-').Skip(0).Take(4));
+				PreflightChangelist = int.Parse(BuildName.Split('-')[4]);
+			}
+			
+			// some platforms can only output lowercase, but the inputs will always be upper case
+			// and that's what we'll compare against to see if a trace is valid
+			PlatformName = FilenameTokens[1].ToLower();
+			
+			DateTime = FilenameTokens[2].ToLower();
+			
+			DeviceProfile = FilenameTokens[3].ToLower();
+			
+			TestName = FilenameTokens[4].ToLower();
+			
+		}
+
+		/// <summary>
+		/// Sanitize the trace's filename to help conform to path-lenght limits
+		/// </summary>
+		/// <returns></returns>
+		public string GetSanitizedFileName()
+		{
+			// reconstruct the filename for this test to skip the preflight identifier
+			string SanitizedFilename = $"{Build}-cl-{Changelist.ToString()}";
+			if (PreflightChangelist > -1)
+			{
+				SanitizedFilename += $"-pf-{PreflightChangelist.ToString()}";
+			}
+			SanitizedFilename += $"_{PlatformName}_{DateTime}_{TestName}";
+
+			return SanitizedFilename;
+		}
+
+		/// <summary>
+		/// Returns the full path to the log file to use when parsing this trace
+		/// </summary>
+		/// <returns></returns>
+		public string GetLogFilePath()
+		{
+			return Path.Combine(FilePath, FileName + ".log");
+		}
+		
+		/// <summary>
+		/// Checks to see if this trace has a matching CSV summary file output
+		/// which would indicate that it has already been processed
+		/// </summary>
+		/// <returns>bool</returns>
+		public bool HasMatchingLog()
+		{
+			return File.Exists(GetLogFilePath());
+		}
+
+		/// <summary>
+		/// Return the directory to which the Insights CLI should output all the region CSVs
+		/// </summary>
+		/// <returns></returns>
+		public string GetRegionCSVOutputPath()
+		{
+			return FilePath;
+		}
+		
+		/// <summary>
+		/// For after trace analysis, find all the output CSVs for all the regions
+		/// </summary>
+		/// <returns></returns>
+		public List<string> GatherOutputCSVs()
+		{
+			List<string> GatheredCSVs = new List<string>();
+			if (Directory.Exists(GetRegionCSVOutputPath()))
+			{
+				GatheredCSVs.AddRange(
+					from CSVFile in Directory.GetFiles(GetRegionCSVOutputPath(), $"{FileName}*.csv",
+						SearchOption.AllDirectories)
+					select CSVFile);
+			}
+
+			return GatheredCSVs;
+		}
+	}
+	public abstract class AutomatedPerfTestInsightsReport : BuildCommand
+	{
+		public string BuildName;
+		public string ProjectName;
+		public string TestName;
+		public string PerfCacheRoot;
+		public bool UseShippingInsights;
+		
+		private void InitParams()
+		{
+			// grab config variables from commandline
+			BuildName = ParseParamValue("AutomatedPerfTest.BuildName", "NotFound").ToLower();
+			ProjectName = ParseParamValue("AutomatedPerfTest.ProjectName", "NotFound");
+			TestName = ParseParamValue("AutomatedPerfTest.TestName", "NotFound");
+			PerfCacheRoot = ParseParamValue("AutomatedPerfTest.PerfCacheRoot", null);
+			UseShippingInsights = ParseParamBool("AutomatedPerfTest.UseShippingInsights", false);
+		}
+
+		public FileReference GetInsightsExecutablePath()
+		{
+			FileReference InsightsPath = UseShippingInsights ? FileReference.Combine(Unreal.EngineDirectory, "Binaries", "Win64", "UnrealInsights-Win64-Shipping.exe") : FileReference.Combine(Unreal.EngineDirectory, "Binaries", "Win64", "UnrealInsights.exe");
+
+			// ReSharper disable once InvertIf
+			if (!FileReference.Exists(InsightsPath))
+			{
+				Log.Logger.LogError("Failed to find perf report utility at this path: \"{InsightsPath}\".", InsightsPath);
+				return null;
+			}
+
+			return InsightsPath;
+		}
+
+		/// <summary>
+		/// Call Unreal Insights CLI to export timer statistics for all the regions in the input trace
+		/// And output the GPU timings for those regions to a CSV
+		/// Then upload the results from those CSVs to the database
+		/// </summary>
+		/// <param name="TraceFile"></param>
+		private void ProcessTraceFile(TraceFileInfo TraceFile)
+		{
+			Log.Logger.LogInformation("Processing trace file {TraceFilePath}", TraceFile.TraceFilePath);
+			
+			// make sure we can output the CSVs to the right folder by creating it first
+			InternalUtils.SafeCreateDirectory(TraceFile.GetRegionCSVOutputPath(), true);
+
+			// only export regions matching the build name, as other systems are starting to use regions and we don't need to export ALL of them
+			CommandUtils.RunAndLog(GetInsightsExecutablePath().FullName,
+				$"-AutoQuit -NoUI -ABSLOG=\"{TraceFile.GetLogFilePath()}\" -OpenTraceFile=\"{TraceFile.TraceFilePath}\" -ExecOnAnalysisCompleteCmd=\"TimingInsights.ExportTimerStatistics {TraceFile.GetRegionCSVOutputPath()}\\*.csv -region={TraceFile.BuildName}* -threads=GPU\"",
+				out int ErrorCode);
+
+			if (ErrorCode != 0)
+			{
+				Log.Logger.LogError(
+					"MegaworldsPerformanceTool returned error code \"{ErrorCode}\" while generating detailed report from trace \"{TraceFile.TraceFilePath}\" CSV folder \"{CSVOutputPath}\"",
+					ErrorCode, TraceFile, TraceFile.GetRegionCSVOutputPath());
+				return;
+			}
+			
+			List<string> CSVFilenames = TraceFile.GatherOutputCSVs();
+
+			if (CSVFilenames.Count < 1)
+			{
+				Log.Logger.LogError("No CSVs found at output path {CSVOutputPath}", TraceFile.GetRegionCSVOutputPath());
+				return;
+			}
+
+			Dictionary<string, dynamic>[] outputJSONData = new Dictionary<string, dynamic>[CSVFilenames.Count];
+			
+			for (int i = 0; i < CSVFilenames.Count; i++)
+			{
+				// create a dict to which we can load the data we want from the CSV and output it to the database
+				Dictionary<string, dynamic> jsonData = new Dictionary<string, dynamic>();
+				
+				// find the specific region id from the end of the CSV region name, this could also just be TestName
+				string RegionID = Path.GetFileNameWithoutExtension(CSVFilenames[i]).Split('_').Last();
+				
+				// pass along the known values
+				jsonData["Build"] = TraceFile.Build;
+				jsonData["Changelist"] = TraceFile.Changelist;
+				jsonData["PreflightCL"] = TraceFile.PreflightChangelist;
+				jsonData["PlatformName"] = TraceFile.PlatformName;
+				jsonData["DateTime"] = TraceFile.DateTime;
+				jsonData["DeviceProfile"] = TraceFile.DeviceProfile;
+				jsonData["IsBuildMachine"] = IsBuildMachine;
+				jsonData["TestName"] = TestName;
+				jsonData["RegionID"] = RegionID;
+
+				// TODO make this a config file
+				
+				// load the CSV to grab the columns we want from the rows of the CSV
+				using (StreamReader csvParser = new StreamReader(CSVFilenames[i]))
+				{
+					// skip the header
+					csvParser.ReadLine();
+
+					string line;
+					string[] fields;
+					string TimerName;
+					string InclusiveAverage;
+
+					while (!csvParser.EndOfStream)
+					{
+						line = csvParser.ReadLine();
+						if (line != null)
+						{
+							fields = line.Split(",");
+							TimerName = fields[0];
+							
+							// the output CSV from Insights puts Inclusive Average at the 6th column
+							InclusiveAverage = fields[5];
+							
+							// multiply the timer number by 1000, since the CSV output is in Seconds and 
+							// we'd much rather see it in Milliseconds
+							jsonData[TimerName] = float.Parse(InclusiveAverage) * 1000.0;
+						}
+					}
+				}
+
+				// output the json data to a string, and then to a file so we can quickly reconstruct the database if necessary
+				string jsonString =
+					JsonSerializer.Serialize(jsonData, new JsonSerializerOptions() {WriteIndented = true});
+
+				string JsonOutfile = CSVFilenames[i].Replace(".csv", ".json");
+
+				Log.Logger.LogInformation("Outputting {CSVFilename} to {JsonOutfile}", CSVFilenames[i], JsonOutfile);
+				
+				InternalUtils.SafeCreateDirectory(Path.GetDirectoryName(JsonOutfile), true);
+				
+				using (StreamWriter outputFile = new StreamWriter(JsonOutfile))
+				{
+					outputFile.WriteLine(jsonString);
+				}
+
+				outputJSONData[i] = jsonData;
+			}
+			
+			OutputSummary(TraceFile, outputJSONData);
+		}
+
+		private void OutputSummary(TraceFileInfo TraceFile, Dictionary<string, dynamic>[] OutputJSONData)
+		{
+			string OutCSVName = TraceFile.TraceFilePath.Replace(".utrace", ".csv");
+
+			string CSVHeader = "RegionName,Unaccounted";
+
+			Log.Logger.LogInformation("Outputting CSV summary to {SummaryCSVPath}", OutCSVName);
+			
+			using (StreamWriter OutCSV = new StreamWriter(OutCSVName))
+			{
+				OutCSV.WriteLine(CSVHeader);
+				Log.Logger.LogInformation(CSVHeader);
+				for (int i = 0; i < OutputJSONData.Length; i++)
+				{
+					string lineString = $"{OutputJSONData[i]["RegionID"]},{OutputJSONData[i]["Unaccounted"]}";
+					OutCSV.WriteLine(lineString);
+					Log.Logger.LogInformation(lineString);
+				}
+			}
+		}
+
+		public override void ExecuteBuild()
+		{
+			InitParams();
+
+			Log.Logger.LogInformation("BuildName from BuildGraph = {BuildName}", BuildName);
+
+			// loop over all of the unprocessed trace files in the perf cache and upload the results
+			// more brute force, but less likely to upload the wrong results or not upload any results
+			
+			// search through the path provided from commandline
+			var DiscoveredTraces = new List<string>();
+			if (Directory.Exists(PerfCacheRoot))
+			{
+				DiscoveredTraces.AddRange(
+					from TraceFile in Directory.GetFiles(PerfCacheRoot, "*.utrace", SearchOption.AllDirectories)
+					select TraceFile);
+			}
+			else
+			{
+				Log.Logger.LogError("PerfCacheRoot does not exists {PerfCacheRoot}", PerfCacheRoot);
+			}
+
+			// if we couldn't find any traces, report that and bail out
+			if (DiscoveredTraces.Count == 0)
+			{
+				Log.Logger.LogError("Tests completed successfully but no trace results were found. Searched path was {PerfCacheRoot}", PerfCacheRoot);
+			}
+			
+			// Sort the cases by timestamp, newest first
+			string[] SortedTraces = 
+				(from TraceFile in DiscoveredTraces
+					let Timestamp = File.GetCreationTimeUtc(TraceFile)
+					orderby Timestamp descending
+					select TraceFile).ToArray();
+			
+			for (int i = 0; i < SortedTraces.Length; i++)
+			{
+				TraceFileInfo TraceInfo = new TraceFileInfo(SortedTraces[i]);
+				
+				if (!TraceInfo.HasMatchingLog())
+				{
+					Log.Logger.LogInformation("Found unprocessed trace file {MatchingTrace}", SortedTraces[i]);
+					ProcessTraceFile(TraceInfo);
+				}
+			}
+		}
+	}
+}
\ No newline at end of file
diff --git a/Engine/Plugins/Performance/AutomatedPerfTesting/Build/Scripts/AutomatedPerfTestNode.cs b/Engine/Plugins/Performance/AutomatedPerfTesting/Build/Scripts/AutomatedPerfTestNode.cs
--- a/Engine/Plugins/Performance/AutomatedPerfTesting/Build/Scripts/AutomatedPerfTestNode.cs
+++ b/Engine/Plugins/Performance/AutomatedPerfTesting/Build/Scripts/AutomatedPerfTestNode.cs
@@ -30,7 +30,7 @@ namespace AutomatedPerfTest
 			// after GetConfiguration is called. This will cause a mismatch in CreateReport.
 			OriginalBuildName = Globals.Params.ParseValue("BuildName", InContext.BuildInfo.BuildName);
 			Log.Info("Setting OriginalBuildName to {OriginalBuildName}", OriginalBuildName);
-
+			
 			TestGuid = Guid.NewGuid();
 			Log.Info("Your Test GUID is :\n" + TestGuid.ToString() + '\n');
 
@@ -122,8 +122,10 @@ namespace AutomatedPerfTest
 			if (App != null)
 			{
 				UnrealLogParser Parser = new UnrealLogParser(App.StdOut);
-
-				string LogChannelName = GetCachedConfiguration().ProjectName + "Test";
+				
+				
+				
+				string LogChannelName = Context.BuildInfo.ProjectName + "Test";
 				List<string> TestLines = Parser.GetLogChannel(LogChannelName).ToList();
 
 				string LogCategory = "Log" + LogChannelName;
@@ -228,7 +230,12 @@ namespace AutomatedPerfTest
 		{
 			if (Result == TestResult.Passed)
 			{
-				if (GetConfiguration().DoCSVProfiler)
+				if (GetCachedConfiguration().DoInsightsTrace)
+				{
+					CopyInsightsTraceToPerfCache(ArtifactPath);
+				}
+				
+				if (GetCurrentPass() <= GetNumPasses() && GetCachedConfiguration().DoCSVProfiler)
 				{
 					// Our artifacts from each iteration such as the client log will be overwritten by subsequent iterations so we need to copy them out to a temp dir
 					// to preserve them until we're ready to make our report on the final iteration.
@@ -249,7 +256,7 @@ namespace AutomatedPerfTest
 						{
 							{"HordeJobUrl", Globals.Params.ParseValue("JobDetails", null)}
 						};
-
+						
 						Log.Info("Creating perf server importer with build name {BuildName}", OriginalBuildName);
 						string DataSourceName = GetConfiguration().DataSourceName;
 						string ImportDirOverride = Globals.Params.ParseValue("PerfReportServerImportDir", null);
@@ -274,15 +281,22 @@ namespace AutomatedPerfTest
 							{
 								string InsightsFilename = Path.GetFileNameWithoutExtension(CsvImportEntries.First().CsvFilename)
 									.Replace(".csv", ".utrace");
-								string InsightsFilePath = Path.Combine(ArtifactPath, "Profiling", InsightsFilename);
-								if (File.Exists(InsightsFilePath))
+
+								// recursively look for trace files that match the CSV's filename in the artifact path
+								string[] MatchingTraces = FindFiles($"*{InsightsFilename}", true, ArtifactPath);
+								if(MatchingTraces.Length > 0)
 								{
-									CsvImportEntries.First().AddAdditionalFile("Insights", InsightsFilePath);
+									if (MatchingTraces.Length > 1)
+									{
+										Log.Warning("Multiple Insights traces were found in {ArtifactPath} matching pattern *{InsightsFilename}. Only the first will be attached to the CSV import for this test.",
+											ArtifactPath, InsightsFilename);										
+									}
+									CsvImportEntries.First().AddAdditionalFile("Insights", MatchingTraces.First());
 								}
 								else
 								{
-									Log.Warning("Insights was requested, but no match insights trace file was found {InsightsFilename}",
-										InsightsFilePath);
+									Log.Warning("Insights was requested, but no matching insights traces were found  matching pattern *{InsightsFilename} in {ArtifactPath}",
+										InsightsFilename, ArtifactPath);
 								}
 							}
 
@@ -334,9 +348,9 @@ namespace AutomatedPerfTest
 				Logger.LogError("Failed to find perf report utility at this path: \"{ToolPath}\".", ToolPath);
 				return;
 			}
-
-			var ReportConfigDir = Path.Combine(Unreal.RootDirectory.FullName, "Samples", "Showcases", GetCachedConfiguration().ProjectName, "Build", "Scripts", "PerfReport");
-			var ReportPath = Path.Combine(ArtifactPath, "Reports", "Performance");
+			
+			var ReportConfigDir = Path.Combine(Context.Options.ProjectPath.Directory.FullName, "Build", "Scripts", "PerfReport");
+			var ReportPath = Path.Combine(ReportCacheDir, "Reports", "Performance");
 
 			// Csv files may have been output in one of two places.
 			// Check both...
@@ -398,8 +412,10 @@ namespace AutomatedPerfTest
 			// Win64 is actually called "Windows" in csv profiles
 			var PlatformNameFilter = Platform == UnrealTargetPlatform.Win64 ? "Windows" : $"{Platform}";
 
+			string SearchPattern = $"{Context.BuildInfo.ProjectName}*";
+			
 			// Produce the detailed report, and update the perf cache
-			CommandUtils.RunAndLog(ToolPath.FullName, $"-csvdir \"{NewestDir}\" -o \"{ReportPath}\" -reportxmlbasedir \"{ReportConfigDir}\" -summaryTableCache \"{ReportCacheDir}\" -searchpattern csvprofile* -metadatafilter platform=\"{PlatformNameFilter}\"", out int ErrorCode);
+			CommandUtils.RunAndLog(ToolPath.FullName, $"-csvdir \"{NewestDir}\" -o \"{ReportPath}\" -reportxmlbasedir \"{ReportConfigDir}\" -summaryTableCache \"{ReportCacheDir}\" -searchpattern {SearchPattern} -metadatafilter platform=\"{PlatformNameFilter}\"", out int ErrorCode);
 			if (ErrorCode != 0)
 			{
 				Logger.LogError("PerfReportTool returned error code \"{ErrorCode}\" while generating detailed report.", ErrorCode);
@@ -503,7 +519,7 @@ namespace AutomatedPerfTest
 				Log.Warning("No valid csv files found in {CSVPath}", CSVPath);
 			}
 		}
-
+		
 		protected virtual string GetSubtestName()
 		{
 			return "Performance";
@@ -517,13 +533,31 @@ namespace AutomatedPerfTest
 			UnrealTestRole ClientRole = Config.RequireRole(UnrealTargetRole.Client);
 			// the controller will be added by the subclasses
 
+			ClientRole.CommandLineParams.AddOrAppendParamValue("logcmds", "LogHttp Verbose, LogAutomatedPerfTest Verbose");
+			
 			ClientRole.CommandLineParams.Add("-deterministic");
 
-			Log.Info("AutomatedPerfTestNode<>.GetConfiguration(): Config.DoFPSChart={0}, Config.DoCSVProfiler={1}, Config.DoVideoCapture={2}", Config.DoFPSChart, Config.DoCSVProfiler, Config.DoVideoCapture);
+			Log.Info("AutomatedPerfTestNode<>.GetConfiguration(): Config.DoFPSChart={0}, Config.DoCSVProfiler={1}, Config.DoVideoCapture={2}, Config.DoInsightsTrace={3}", Config.DoFPSChart, Config.DoCSVProfiler, Config.DoVideoCapture, Config.DoInsightsTrace);
+			
+			ClientRole.CommandLineParams.AddOrAppendParamValue("AutomatedPerfTest.TestName", Config.TestName);
+
+			if (Config.DeviceProfileOverride != String.Empty)
+			{
+				ClientRole.CommandLineParams.AddOrAppendParamValue("AutomatedPerfTest.DeviceProfileOverride", Config.DeviceProfileOverride);
+			}
 
+			if (Config.DoInsightsTrace)
+			{
+				ClientRole.CommandLineParams.Add("AutomatedPerfTest.DoInsightsTrace");
+				if (Config.TraceChannels != String.Empty)
+				{
+					ClientRole.CommandLineParams.AddOrAppendParamValue("AutomatedPerfTest.TraceChannels", Config.TraceChannels);
+				}
+			}
+			
 			if (string.IsNullOrEmpty(Config.DataSourceName))
 			{
-				Config.DataSourceName = string.Format("Automation.{0}.Performance", Config.ProjectName);
+				Config.DataSourceName = string.Format("Automation.{0}.Performance", Context.BuildInfo.ProjectName);
 			}
 
 			if (Config.DoFPSChart)
@@ -539,7 +573,7 @@ namespace AutomatedPerfTest
 				// Add CSV metadata
 				List<string> CsvMetadata = new List<string>
 				{
-					string.Format("testname={0}", Config.ProjectName),
+					string.Format("testname={0}", Context.BuildInfo.ProjectName),
 					"gauntletTestType=AutomatedPerfTest",
 					string.Format("gauntletSubTest={0}", GetSubtestName()),
 					"testBuildIsPreflight=" + (ReportGenUtils.IsTestingPreflightBuild(OriginalBuildName) ? "1" : "0"),
@@ -562,6 +596,57 @@ namespace AutomatedPerfTest
 
 			return Config;
 		}
+		
+		public void CopyInsightsTraceToPerfCache(string ArtifactPath)
+		{
+			Logger.LogInformation("Copying test insights trace from artifact path to report cache");
+			
+			// find all the available trace paths
+			var DiscoveredTraces = new List<string>();
+			if (Directory.Exists(ArtifactPath))
+			{
+				DiscoveredTraces.AddRange(
+					from TraceFile in Directory.GetFiles(ArtifactPath, "*.utrace", SearchOption.AllDirectories)
+					select TraceFile);
+			}
+			
+			// if we couldn't find any traces, report that and bail out
+			if (DiscoveredTraces.Count == 0)
+			{
+				Logger.LogError("Test completed successfully but no trace results were found. Searched path was {ArtifactPath}", ArtifactPath);
+				return;
+			}
+			
+			// iterate over each of the discovered traces (there should be one for each test case that was run)
+			// first, sort the cases by timestamp
+			string[] SortedTraces =
+				(from TraceFile in DiscoveredTraces
+					let Timestamp = File.GetCreationTimeUtc(TraceFile)
+					orderby Timestamp descending
+					select TraceFile).ToArray();
+			
+			var ReportPath = GetCachedConfiguration().PerfCacheRoot;
+			
+			if (SortedTraces.Length > 0)
+			{
+				string Filename = Path.GetFileNameWithoutExtension(SortedTraces[0]);
+				string PerfCachePath = Path.Combine(ReportPath, Filename + ".utrace");
+				
+				Logger.LogInformation("Copying latest utrace file from {ArtifactPath} to perf cache: {PerfCachePath}", ArtifactPath,
+					PerfCachePath);
+				
+				// just try the copy over, and log a failure, but don't bail out of the test.
+				try
+				{
+					InternalUtils.SafeCreateDirectory(PerfCachePath, true);
+					File.Copy(SortedTraces[0], PerfCachePath);
+				}
+				catch (Exception e)
+				{
+					Logger.LogWarning("Failed to copy local trace file: {Text}", e);
+				}
+			}
+		}
 
 		/// <summary>
 		/// Returns the cached version of our config. Avoids repeatedly calling GetConfiguration() on derived nodes
@@ -583,7 +668,7 @@ namespace AutomatedPerfTest
 	/// </summary>
 	/// <typeparam name="TConfigClass"></typeparam>
 	public abstract class AutomatedSequencePerfTestNode<TConfigClass> : AutomatedPerfTestNode<TConfigClass>
-		where TConfigClass : AutomatedPerfTestConfigBase, new()
+		where TConfigClass : AutomatedSequencePerfTestConfig, new()
 	{
 		public AutomatedSequencePerfTestNode(UnrealTestContext InContext) : base(InContext)
 		{
@@ -592,9 +677,23 @@ namespace AutomatedPerfTest
 		public override TConfigClass GetConfiguration()
 		{
 			TConfigClass Config = base.GetConfiguration();
-
-			UnrealTestRole ClientRole = Config.RequireRole(UnrealTargetRole.Client);	// should get an existing role
-			ClientRole.Controllers.Add("AutomatedSequencePerfTest");
+			
+			// extend the role(s) that we initialized in the base class
+			if (Config.GetRequiredRoles(UnrealTargetRole.Client).Any())
+			{
+				foreach(UnrealTestRole ClientRole in Config.GetRequiredRoles(UnrealTargetRole.Client))
+				{
+					ClientRole.Controllers.Add("AutomatedSequencePerfTest");
+					
+					// if a specific MapSequenceComboName was defined in the commandline to UAT, then add that to the commandline for the role
+					if (!string.IsNullOrEmpty(Config.MapSequenceComboName))
+					{
+						// use add Unique, since there should only ever be one of these specified
+						ClientRole.CommandLineParams.AddUnique($"AutomatedPerfTest.SequencePerfTest.MapSequenceName",
+							Config.MapSequenceComboName);
+					}
+				}
+			}
 
 			return Config;
 		}
diff --git a/Engine/Plugins/Performance/AutomatedPerfTesting/Source/AutomatedPerfTesting/AutomatedPerfTesting.Build.cs b/Engine/Plugins/Performance/AutomatedPerfTesting/Source/AutomatedPerfTesting/AutomatedPerfTesting.Build.cs
--- a/Engine/Plugins/Performance/AutomatedPerfTesting/Source/AutomatedPerfTesting/AutomatedPerfTesting.Build.cs
+++ b/Engine/Plugins/Performance/AutomatedPerfTesting/Source/AutomatedPerfTesting/AutomatedPerfTesting.Build.cs
@@ -37,7 +37,7 @@ public class AutomatedPerfTesting : ModuleRules
 				"CoreUObject",
 				"Engine",
 				"Slate",
-				"SlateCore", "LevelSequence", "MovieScene"
+				"SlateCore", "LevelSequence", "MovieScene", "DeveloperSettings"
 				// ... add private dependencies that you statically link with here ...	
 			}
 			);
diff --git a/Engine/Plugins/Performance/AutomatedPerfTesting/Source/AutomatedPerfTesting/Private/AutomatedPerfTestControllerBase.cpp b/Engine/Plugins/Performance/AutomatedPerfTesting/Source/AutomatedPerfTesting/Private/AutomatedPerfTestControllerBase.cpp
--- a/Engine/Plugins/Performance/AutomatedPerfTesting/Source/AutomatedPerfTesting/Private/AutomatedPerfTestControllerBase.cpp
+++ b/Engine/Plugins/Performance/AutomatedPerfTesting/Source/AutomatedPerfTesting/Private/AutomatedPerfTestControllerBase.cpp
@@ -8,6 +8,7 @@
 #include "DeviceProfiles/DeviceProfileManager.h"
 #include "GameFramework/GameModeBase.h"
 #include "AutomatedPerfTestInterface.h"
+#include "AutomatedPerfTestProjectSettings.h"
 #include "Misc/App.h"
 #include "Misc/CommandLine.h"
 #include "ProfilingDebugging/TraceAuxiliary.h"
@@ -20,35 +21,10 @@
 #include "Engine/Engine.h"
 #include "Misc/Paths.h"
 
+
 DEFINE_LOG_CATEGORY(LogAutomatedPerfTest)
 CSV_DEFINE_CATEGORY(AutomatedPerfTest, true);
 
-namespace AutomatedPerfTesting
-{
-	static TAutoConsoleVariable<FString> CVarExplicitTestID(
-		TEXT("AutomatedPerfTest.ExplicitTestID"),
-		"",
-		TEXT("Overrides the automatically generated test ID with the one provided"),
-		ECVF_Default);
-
-	static TAutoConsoleVariable<FString> CVarTraceChannels(
-		TEXT("AutomatedPerfTest.TraceChannels"),
-		"default,screenshot,stats",
-		TEXT("Sets the channels that will be traced when the AutomatedPerfTest is run. Defaults to default,screenshot,stats"),
-		ECVF_Default);
-
-	static TAutoConsoleVariable<FString> CVarDeviceProfileOverride(
-		TEXT("AutomatedPerfTest.DeviceProfileOverride"),
-		"",
-		TEXT("Overrides device profile used when running automated performance tests if set"),
-		ECVF_Default);
-
-	static TAutoConsoleVariable<FString> CVarTestName(
-		TEXT("AutomatedPerfTest.TestName"),
-		"Automated",
-		TEXT("User-defined ID of the test being run. Helpful for identifying test specifics eg Automated, TestingSomeVariable, MyCoolTest"),
-		ECVF_Default);
-}
 
 void UAutomatedPerfTestControllerBase::OnPreWorldInitializeInternal(UWorld* World, const UWorld::InitializationValues IVS)
 {
@@ -126,45 +102,24 @@ UAutomatedPerfTestControllerBase::UAutomatedPerfTestControllerBase(const FObject
 
 FString UAutomatedPerfTestControllerBase::GetTestName()
 {
-	return AutomatedPerfTesting::CVarTestName->GetString();
+	return TestName;
 }
 
 FString UAutomatedPerfTestControllerBase::GetDeviceProfile()
 {
-	if(!AutomatedPerfTesting::CVarDeviceProfileOverride->GetString().IsEmpty())
-	{
-		return AutomatedPerfTesting::CVarDeviceProfileOverride->GetString();
-	}
-	
-	return UDeviceProfileManager::Get().GetActiveDeviceProfileName(); 
-}
-
-FString UAutomatedPerfTestControllerBase::GetExplicitTestID()
-{
-	// if the user has explicitly set a test ID, then return that
-	if(!AutomatedPerfTesting::CVarExplicitTestID->GetString().IsEmpty())
-	{
-		return AutomatedPerfTesting::CVarExplicitTestID->GetString();
-	}
-	
-	return TEXT("");
+	return DeviceProfileOverride.IsEmpty() ? UDeviceProfileManager::Get().GetActiveDeviceProfileName() : DeviceProfileOverride;  
 }
 
 FString UAutomatedPerfTestControllerBase::GetTestID()
 {
-	// if the user has explicitly set a test ID, then return that
-	if(!AutomatedPerfTesting::CVarExplicitTestID->GetString().IsEmpty())
-	{
-		return AutomatedPerfTesting::CVarExplicitTestID->GetString();
-	}
-
 	const TArray<FString> TestCaseIDElements = {FApp::GetBuildVersion(),
 										  FPlatformProperties::PlatformName(),
 										  TestDatetime,
 								          *GetDeviceProfile(),
 								          *GetTestName()};
-	
-	// otherwise construct a unique ID of the form BuildVersion_PlatformName_YYYYMMDD-HHMMSS_DeviceProfile_TestName
+
+	// TODO make this format definable in project settings, similar to how MovieRenderQueue does it for renders
+	// construct a unique ID of the form BuildVersion_PlatformName_YYYYMMDD-HHMMSS_DeviceProfile_TestName
 	FString TestCaseID = FString::Join(TestCaseIDElements, TEXT("_"));
 
 	return TestCaseID;
@@ -175,6 +130,11 @@ FString UAutomatedPerfTestControllerBase::GetOverallRegionName()
 	return GetTestID() + "_" + "Overall";
 }
 
+FString UAutomatedPerfTestControllerBase::GetTraceChannels()
+{
+	return TraceChannels;
+}
+
 bool UAutomatedPerfTestControllerBase::RequestsInsightsTrace() const
 {
 	return bRequestsInsightsTrace;
@@ -198,8 +158,8 @@ bool UAutomatedPerfTestControllerBase::RequestsVideoCapture() const
 bool UAutomatedPerfTestControllerBase::TryStartInsightsTrace()
 {
 	const FString TraceFileName = GetTestID() + ".utrace";
-	UE_LOG(LogAutomatedPerfTest, Log, TEXT("Attempting to start insights trace to file %s with channels %s"), *TraceFileName, *AutomatedPerfTesting::CVarTraceChannels->GetString());
-	return FTraceAuxiliary::Start(FTraceAuxiliary::EConnectionType::File, *TraceFileName, *AutomatedPerfTesting::CVarTraceChannels->GetString());
+	UE_LOG(LogAutomatedPerfTest, Log, TEXT("Attempting to start insights trace to file %s with channels %s"), *TraceFileName, *GetTraceChannels());
+	return FTraceAuxiliary::Start(FTraceAuxiliary::EConnectionType::File, *TraceFileName, *GetTraceChannels());
 }
 
 bool UAutomatedPerfTestControllerBase::TryStopInsightsTrace()
@@ -384,7 +344,7 @@ void UAutomatedPerfTestControllerBase::RunTest()
 	}
 }
 
-void UAutomatedPerfTestControllerBase::TeardownTest()
+void UAutomatedPerfTestControllerBase::TeardownTest(bool bExitAfterTeardown)
 {
 	UE_LOG(LogAutomatedPerfTest, Log, TEXT("Base:: TeardownTest"));
 	
@@ -417,10 +377,17 @@ void UAutomatedPerfTestControllerBase::TeardownTest()
 		IAutomatedPerfTestInterface::Execute_TeardownTest(GameMode);
 	}
 
+	if(bExitAfterTeardown)
+	{
+		TriggerExitAfterDelay();
+	}
+}
 
+void UAutomatedPerfTestControllerBase::TriggerExitAfterDelay()
+{
 	FTimerHandle UnusedHandle;
-	// TODO parameterize the teardown to Exit delay
-	GetWorld()->GetTimerManager().SetTimer(UnusedHandle, this, &UAutomatedPerfTestControllerBase::Exit, 1.0, false, 5.0);
+	
+	GetWorld()->GetTimerManager().SetTimer(UnusedHandle, this, &UAutomatedPerfTestControllerBase::Exit, 1.0, false, GetDefault<UAutomatedPerfTestProjectSettings>()->TeardownToExitDelay);	
 }
 
 void UAutomatedPerfTestControllerBase::Exit()
@@ -443,8 +410,8 @@ void UAutomatedPerfTestControllerBase::Exit()
 			UE_LOG(LogAutomatedPerfTest, Log, TEXT("CSVProfile requested, and test is exiting, but CSV Profiler isn't done writing."))
 			// if we requested a CSV Profile, and the CSV Profiler is still writing the file, add a lambda to call the EndAutomatedPerfTest function
 			// so that we don't exit out of the application before the CSV Profiler is done
-			// TODO there's probably a nicer way to do this
-			// TODO might not need to do this if we set csv.BlockOnCaptureEnd
+			
+			// TODO investigate if this is necessary if we set csv.BlockOnCaptureEnd
 			FCsvProfiler::Get()->OnCSVProfileFinished().AddLambda([this](const FString& Filename)
 			{
 				EndAutomatedPerfTest();
@@ -457,33 +424,22 @@ void UAutomatedPerfTestControllerBase::Exit()
 	EndAutomatedPerfTest();
 }
 
+AGameModeBase* UAutomatedPerfTestControllerBase::GetGameMode() const
+{
+	return GameMode;
+}
+
 void UAutomatedPerfTestControllerBase::OnInit()
 {
 	Super::OnInit();
 
 	UE_LOG(LogAutomatedPerfTest, Log, TEXT("Base:: OnInit"));
 
-	if (FParse::Value(FCommandLine::Get(), TEXT("AutomatedPerfTest.ExplicitTestID="), ExplicitTestID))
-	{
-		AutomatedPerfTesting::CVarExplicitTestID->Set(*ExplicitTestID);
-	}
-
 	// don't stop on separator because this will come in comma-separated
-	if (FParse::Value(FCommandLine::Get(), TEXT("AutomatedPerfTest.TraceChannels="), TraceChannels, false))
-	{
-		AutomatedPerfTesting::CVarTraceChannels->Set(*TraceChannels);
-	}
-
-	if (FParse::Value(FCommandLine::Get(), TEXT("AutomatedPerfTest.DeviceProfileOverride="), DeviceProfileOverride))
-	{
-		AutomatedPerfTesting::CVarDeviceProfileOverride->Set(*DeviceProfileOverride);
-	}
-
-	if (FString TestName; FParse::Value(FCommandLine::Get(), TEXT("AutomatedPerfTest.TestName="), TestID))
-	{
+	FParse::Value(FCommandLine::Get(), TEXT("AutomatedPerfTest.TraceChannels="), TraceChannels, false);
 	
-		AutomatedPerfTesting::CVarTestName->Set(*TestID);
-	}
+	FParse::Value(FCommandLine::Get(), TEXT("AutomatedPerfTest.DeviceProfileOverride="), DeviceProfileOverride);
+	FParse::Value(FCommandLine::Get(), TEXT("AutomatedPerfTest.TestName="), TestName);
 	
 	if(FParse::Param(FCommandLine::Get(), TEXT("AutomatedPerfTest.DoInsightsTrace")))
 	{
diff --git a/Engine/Plugins/Performance/AutomatedPerfTesting/Source/AutomatedPerfTesting/Private/AutomatedPerfTestProjectSettings.cpp b/Engine/Plugins/Performance/AutomatedPerfTesting/Source/AutomatedPerfTesting/Private/AutomatedPerfTestProjectSettings.cpp
new file mode 100644
--- /dev/null
+++ b/Engine/Plugins/Performance/AutomatedPerfTesting/Source/AutomatedPerfTesting/Private/AutomatedPerfTestProjectSettings.cpp
@@ -0,0 +1,8 @@
+// Copyright Epic Games, Inc. All Rights Reserved.
+
+#include "AutomatedPerfTestProjectSettings.h"
+
+UAutomatedPerfTestProjectSettings::UAutomatedPerfTestProjectSettings(const FObjectInitializer& Initializer)
+	: Super(Initializer)
+{
+}
\ No newline at end of file
diff --git a/Engine/Plugins/Performance/AutomatedPerfTesting/Source/AutomatedPerfTesting/Private/AutomatedSequencePerfTest.cpp b/Engine/Plugins/Performance/AutomatedPerfTesting/Source/AutomatedPerfTesting/Private/AutomatedSequencePerfTest.cpp
--- a/Engine/Plugins/Performance/AutomatedPerfTesting/Source/AutomatedPerfTesting/Private/AutomatedSequencePerfTest.cpp
+++ b/Engine/Plugins/Performance/AutomatedPerfTesting/Source/AutomatedPerfTesting/Private/AutomatedSequencePerfTest.cpp
@@ -7,67 +7,120 @@
 #include "AutomatedPerfTesting.h"
 #include "LevelSequenceActor.h"
 #include "LevelSequencePlayer.h"
-#include "MovieSceneSequence.h"
 #include "MovieSceneSequencePlayer.h"
+#include "Kismet/GameplayStatics.h"
 #include "Misc/CommandLine.h"
 
-namespace AutomatedPerfTesting
+/*****
+ * UAutomatedSequencePerfTestProjectSettings
+ *****/
+
+UAutomatedSequencePerfTestProjectSettings::UAutomatedSequencePerfTestProjectSettings(const FObjectInitializer& Initializer)
+	: Super(Initializer)
 {
-	static TAutoConsoleVariable<FString> CvarSequencePath(
-	TEXT("AutomatedPerfTest.SequencePath"),
-	"",
-	TEXT("Full path to the sequence which should be run when the AutomatedSequencePerfTest is executed. EG. /Game/Profiling/SampleSequence.SampleSequence"),
-	ECVF_Default);
 }
 
-FString UAutomatedSequencePerfTest::GetSequenceName() const
+bool UAutomatedSequencePerfTestProjectSettings::GetComboFromTestName(FName TestName, FAutomatedPerfTestMapSequenceCombo& FoundSequence) const
 {
-	return SequenceSoftPath.GetAssetName();
+	for (FAutomatedPerfTestMapSequenceCombo TestCombo : MapsAndSequencesToTest)
+	{
+		if(TestCombo.ComboName == TestName)
+		{
+			FoundSequence = TestCombo;
+			return true;
+		}
+	}
+	return false;
 }
 
+
+/*****
+ * UAutomatedSequencePerfTest
+ *****/
 void UAutomatedSequencePerfTest::SetupTest()
 {
-	Super::SetupTest();
-
-	UE_LOG(LogAutomatedPerfTest, Log, TEXT("Setting up test"));
-
-	// make sure the world exists, then create a sequence player
-	if(UWorld* const World = GetWorld())
+	if(&CurrentMapSequenceCombo != nullptr)
 	{
-		// load the sequence specified by the user
-		UE_LOG(LogAutomatedPerfTest, Log, TEXT("Loading sequence %s"), *SequenceSoftPath.ToString());
-		ULevelSequence* TargetSequence = LoadObject<ULevelSequence>(NULL, *SequenceSoftPath.ToString(), NULL, LOAD_None, NULL);
-		check(TargetSequence);
+		if(GetCurrentMap() == CurrentMapSequenceCombo.Map.GetAssetName())
+		{
+			// don't even try to set up the test if we're not in the correct map
+			Super::SetupTest();
+			
+			// make sure the world exists, then create a sequence player
+			if(UWorld* const World = GetWorld())
+			{
+				// load the sequence specified by the user
+				UE_LOG(LogAutomatedPerfTest, Log, TEXT("Loading sequence %s"), *CurrentMapSequenceCombo.Sequence.ToString());
+				ULevelSequence* TargetSequence = LoadObject<ULevelSequence>(NULL, *CurrentMapSequenceCombo.Sequence.ToString(), NULL, LOAD_None, NULL);
+				check(TargetSequence);
 		
-		UE_LOG(LogAutomatedPerfTest, Log, TEXT("World is valid, creating sequence player"));
-		SequencePlayer = ULevelSequencePlayer::CreateLevelSequencePlayer(World, TargetSequence, FMovieSceneSequencePlaybackSettings(), SequenceActor);
+				UE_LOG(LogAutomatedPerfTest, Log, TEXT("World is valid, creating sequence player"));
+				SequencePlayer = ULevelSequencePlayer::CreateLevelSequencePlayer(World, TargetSequence, FMovieSceneSequencePlaybackSettings(), SequenceActor);
+
+				if (SequencePlayer == nullptr)
+				{
+					UE_LOG(LogAutomatedPerfTest, Error, TEXT("Unable to create sequence player when starting AutomatedSequencePerfTest, exiting..."));
+					EndAutomatedPerfTest(1);
+				}
 
-		if (SequencePlayer == nullptr)
+				// set the sequence up at the beginning
+				FMovieSceneSequencePlaybackParams PlaybackParams = FMovieSceneSequencePlaybackParams();
+				FMovieSceneSequencePlayToParams PlayToParams = FMovieSceneSequencePlayToParams();
+
+				PlaybackParams.Time = 0.0;
+				PlaybackParams.UpdateMethod = EUpdatePositionMethod::Scrub;
+
+				UE_LOG(LogAutomatedPerfTest, Log, TEXT("SetupMapTest:: Scrubbing to start"));
+				SequencePlayer->PlayTo(PlaybackParams, PlayToParams);
+
+				UE_LOG(LogAutomatedPerfTest, Verbose, TEXT("SetupMapTest:: Waiting for %f seconds before playing seqeuence"), Settings->SequenceStartDelay);
+				FTimerHandle UnusedHandle;
+				GetWorld()->GetTimerManager().SetTimer(UnusedHandle, this, &UAutomatedSequencePerfTest::RunTest, 1.0, false, Settings->SequenceStartDelay);
+			}
+			// if we have an invalid world, we can't run the test, so we should bail out
+			else
+			{
+				UE_LOG(LogAutomatedPerfTest, Error, TEXT("Invalid World when starting AutomatedSequencePerfTest, exiting..."));
+				EndAutomatedPerfTest(1);
+			}
+		}
+		else
 		{
-			UE_LOG(LogAutomatedPerfTest, Error, TEXT("Unable to create sequence player when starting AutomatedSequencePerfTest, exiting..."));
-			EndAutomatedPerfTest(1);
+			UE_LOG(LogAutomatedPerfTest, Log, TEXT("Current Map Name %s is not expected %s, calling NextMap."), *GetCurrentMap(), *CurrentMapSequenceCombo.Map.GetAssetName())
+			NextMap();
 		}
+	}
+	else
+	{
+		UE_LOG(LogAutomatedPerfTest, Log, TEXT("Current Map Sequence Combo has not been set, calling NextMap to trigger."))
+		NextMap();
+	}
+}
 
-		// set the sequence up at the beginning
-		FMovieSceneSequencePlaybackParams PlaybackParams = FMovieSceneSequencePlaybackParams();
-		FMovieSceneSequencePlayToParams PlayToParams = FMovieSceneSequencePlayToParams();
+void UAutomatedSequencePerfTest::NextMap()
+{
+	UE_LOG(LogAutomatedPerfTest, Log, TEXT("UAutomatedSequencePerfTest::NextMap"))
 
-		PlaybackParams.Time = 0.0;
-		PlaybackParams.UpdateMethod = EUpdatePositionMethod::Scrub;
+	if(MapSequenceCombos.Num() > 0)
+	{
+		CurrentMapSequenceCombo = *MapSequenceCombos.Pop();
+		UE_LOG(LogAutomatedPerfTest, Log, TEXT("Setting up test for Map/Sequence combo %s"), *CurrentMapSequenceCombo.ComboName.ToString())
 
-		UE_LOG(LogAutomatedPerfTest, Log, TEXT("RunTest:: Scrubbing to start"));
-		SequencePlayer->PlayTo(PlaybackParams, PlayToParams);
-		
-		FTimerHandle UnusedHandle;
+		// no need to prepend this with a ? since OpenLevel handles that part for us
+		FString OptionsString;
+		if(!CurrentMapSequenceCombo.GameModeOverride.IsEmpty())
+		{
+			UE_LOG(LogAutomatedPerfTest, Log, TEXT("Game Mode overridden to %s"), *CurrentMapSequenceCombo.GameModeOverride)
+			OptionsString += "game=" + CurrentMapSequenceCombo.GameModeOverride;
+		}
 		
-		// TODO parameterize the presoak delay
-		GetWorld()->GetTimerManager().SetTimer(UnusedHandle, this, &UAutomatedSequencePerfTest::RunTest, 1.0, false, 5.0);
+		UE_LOG(LogAutomatedPerfTest, Log, TEXT("Opening map %s%s"), *CurrentMapSequenceCombo.Map.GetAssetName(), *OptionsString);
+		UGameplayStatics::OpenLevel(AutomatedPerfTest::FindCurrentWorld(), *CurrentMapSequenceCombo.Map.GetAssetName(), true, OptionsString);
 	}
-	// if we have an invalid world, we can't run the test, so we should bail out
 	else
 	{
-		UE_LOG(LogAutomatedPerfTest, Error, TEXT("Invalid World when starting AutomatedSequencePerfTest, exiting..."));
-		EndAutomatedPerfTest(1);
+		UE_LOG(LogAutomatedPerfTest, Log, TEXT("UAutomatedSequencePerfTest::NextMap, all maps complete, exiting after delay."))
+		TriggerExitAfterDelay();
 	}
 }
 
@@ -81,10 +134,16 @@ void UAutomatedSequencePerfTest::RunTest()
 	if(SequencePlayer)
 	{
 		UE_LOG(LogAutomatedPerfTest, Log, TEXT("RunTest::Valid Sequence Player, proceeding"));
-		
+
+		// trigger a camera cut manually in order to start the region for the first camera cut
+		OnCameraCut(SequencePlayer->GetActiveCameraComponent());
         SequencePlayer->Play();
-		// TODO add a bit here that also triggers begin/end regions on camera cut
-        SequencePlayer->OnFinished.AddDynamic(this, &UAutomatedSequencePerfTest::TeardownTest);
+		
+		// When the sequence has finished, we'll tear down the test in this map via the OnSequenceFinished dispatch
+		// because TeardownTest's signature doesn't match OnFinished
+        SequencePlayer->OnFinished.AddDynamic(this, &UAutomatedSequencePerfTest::OnSequenceFinished);
+		
+		SequencePlayer->OnCameraCut.AddDynamic(this, &UAutomatedSequencePerfTest::OnCameraCut);
 	}
 	// otherwise bail out of the test
 	else
@@ -94,11 +153,28 @@ void UAutomatedSequencePerfTest::RunTest()
 	}
 }
 
-void UAutomatedSequencePerfTest::TeardownTest()
+void UAutomatedSequencePerfTest::OnSequenceFinished()
+{
+	// trigger OnCameraCut again with a nullptr for the new camera so that
+	// we can end the final camera cut's region
+	OnCameraCut(nullptr);
+	TeardownTest(false);
+}
+
+void UAutomatedSequencePerfTest::TeardownTest(bool bExitAfterTeardown)
 {
 	UE_LOG(LogAutomatedPerfTest, Log, TEXT("AutomatedSequencePerfTest::TeardownTest"));
 	
-	Super::TeardownTest();
+	Super::TeardownTest(bExitAfterTeardown);
+
+	UnbindAllDelegates();
+
+	// null out the references we have to our world objects
+	CurrentCamera = nullptr;
+	SequencePlayer = nullptr;
+	SequenceActor = nullptr;
+	
+	NextMap();
 }
 
 void UAutomatedSequencePerfTest::Exit()
@@ -107,19 +183,86 @@ void UAutomatedSequencePerfTest::Exit()
 	Super::Exit();
 }
 
+void UAutomatedSequencePerfTest::OnCameraCut(UCameraComponent* CameraComponent)
+{
+	// null check the CurrentCamera so that we can use OnCameraCut to mark the starting camera cut region
+	if(CurrentCamera)
+	{
+		if(RequestsInsightsTrace())
+		{
+			TRACE_END_REGION(*GetCameraCutID());
+		}
+#if CSV_PROFILER
+		if(RequestsCSVProfiler())
+		{
+			CSV_EVENT(AutomatedPerfTest, TEXT("END_%s"), *GetCameraCutID())
+		}
+#endif
+	}
+
+	// then null check the new CameraComponent so that we can use OnCameraCut to mark the end of the final camera cut region
+	if(CameraComponent)
+	{
+		// Then bring in the new camera component for this cut and mark the start of it
+		CurrentCamera = CameraComponent;
+		if(RequestsInsightsTrace())
+		{
+			TRACE_BEGIN_REGION(*GetCameraCutID());
+		}
+#if CSV_PROFILER
+		if(RequestsCSVProfiler())
+		{
+			CSV_EVENT(AutomatedPerfTest, TEXT("START_%s"), *GetCameraCutID())
+		}
+#endif
+	}
+}
+
+FString UAutomatedSequencePerfTest::GetTestID()
+{
+	return &CurrentMapSequenceCombo != nullptr ? Super::GetTestID() + "_" + CurrentMapSequenceCombo.ComboName.ToString() : Super::GetTestID();
+}
+
+FString UAutomatedSequencePerfTest::GetCameraCutID()
+{
+	if(CurrentCamera != nullptr)
+	{
+		const AActor* Owner = CurrentCamera->GetOwner();
+		FString CameraCutName = Owner ? Owner->GetActorNameOrLabel() : CurrentCamera->GetReadableName();
+		return GetTestID() + "_" + CameraCutName; 
+	}
+	return GetTestID();
+}
+
 void UAutomatedSequencePerfTest::OnInit()
 {
 	Super::OnInit();
 	UE_LOG(LogAutomatedPerfTest, Log, TEXT("AutomatedSequencePerfTest::OnInit"));
 	
-	if (FParse::Value(FCommandLine::Get(), TEXT("AutomatedPerfTest.SequencePath="), SequencePathName))
+	Settings = GetDefault<UAutomatedSequencePerfTestProjectSettings>();
+
+	// if an explicit map/sequence name was set from commandline, use this to override the test
+	if (FParse::Value(FCommandLine::Get(), TEXT("AutomatedPerfTest.SequencePerfTest.MapSequenceName="), SequenceTestName))
+	{
+		if(Settings->GetComboFromTestName(SequenceTestName, CurrentMapSequenceCombo))
+		{
+			MapSequenceCombos.Add(&CurrentMapSequenceCombo);
+		}
+		else
+		{
+			UE_LOG(LogAutomatedPerfTest, Error, TEXT("Failed to find MapSequence combo name matching %s"), *SequenceTestName.ToString())
+		}
+	}
+	// otherwise, use all the maps defined in project settings
+	else
 	{
-		SequencePathName.ReplaceInline(TEXT("\""), TEXT(""));
-		SequencePathName.ReplaceInline(TEXT("="), TEXT(""));
-		AutomatedPerfTesting::CvarSequencePath->Set(*SequencePathName);
-		UE_LOG(LogAutomatedPerfTest, Log, TEXT("Input Sequence Path Name = %s"), *SequencePathName);
-		SequenceSoftPath = FSoftObjectPath(SequencePathName);
+		for(FAutomatedPerfTestMapSequenceCombo MapSequenceCombo : Settings->MapsAndSequencesToTest)
+		{
+			MapSequenceCombos.Add(&MapSequenceCombo);
+		}
 	}
+	
+	UE_LOG(LogAutomatedPerfTest, Log, TEXT("Initialized with %hhd MapSequence combos"), MapSequenceCombos.Num());
 }
 
 void UAutomatedSequencePerfTest::UnbindAllDelegates()
@@ -129,13 +272,11 @@ void UAutomatedSequencePerfTest::UnbindAllDelegates()
 	// if we have a valid sequence player, make sure we unbind our events from it when we're wrapping up the test.
 	if(SequencePlayer != nullptr)
 	{
+		SequencePlayer->OnCameraCut.RemoveAll(this);
 		SequencePlayer->OnFinished.RemoveAll(this);
+		GetWorld()->GetTimerManager().ClearAllTimersForObject(SequencePlayer);
 	}
 
-	// clear any stray timers that might be laying around
+	// clear any stray timers that might be lying around
 	GetWorld()->GetTimerManager().ClearAllTimersForObject(this);
-	if(SequencePlayer)
-	{
-		GetWorld()->GetTimerManager().ClearAllTimersForObject(SequencePlayer);
-	}
 }
diff --git a/Engine/Plugins/Performance/AutomatedPerfTesting/Source/AutomatedPerfTesting/Public/AutomatedPerfTestControllerBase.h b/Engine/Plugins/Performance/AutomatedPerfTesting/Source/AutomatedPerfTesting/Public/AutomatedPerfTestControllerBase.h
--- a/Engine/Plugins/Performance/AutomatedPerfTesting/Source/AutomatedPerfTesting/Public/AutomatedPerfTestControllerBase.h
+++ b/Engine/Plugins/Performance/AutomatedPerfTesting/Source/AutomatedPerfTesting/Public/AutomatedPerfTestControllerBase.h
@@ -5,11 +5,39 @@
 #include "CoreMinimal.h"
 #include "GauntletTestController.h"
 #include "Engine/World.h"
+#include "Engine/Engine.h"
 
 #include "AutomatedPerfTestControllerBase.generated.h"
 
 class AGameModeBase;
 
+
+namespace AutomatedPerfTest
+{
+	static UWorld* FindCurrentWorld()
+	{
+		UWorld* World = nullptr;
+		for (const FWorldContext& WorldContext : GEngine->GetWorldContexts())
+		{
+			if (WorldContext.WorldType == EWorldType::Game)
+			{
+				World = WorldContext.World();
+			}
+#if WITH_EDITOR
+			else if (GIsEditor && WorldContext.WorldType == EWorldType::PIE)
+			{
+				World = WorldContext.World();
+				if (World)
+				{
+					return World;
+				}
+			}
+#endif
+		}
+
+		return World;
+	}
+}
 /**
  * 
  */
@@ -34,11 +62,11 @@ public:
 public:
 	UAutomatedPerfTestControllerBase(const FObjectInitializer& ObjectInitializer);
 
-	static FString GetTestName();
-	static FString GetDeviceProfile();
-	static FString GetExplicitTestID();
-	FString GetTestID();
+	FString GetTestName();
+	FString GetDeviceProfile();
+	virtual FString GetTestID();
 	FString GetOverallRegionName();
+	FString GetTraceChannels();
 	
 	bool RequestsInsightsTrace() const;
 	bool RequestsCSVProfiler() const;
@@ -59,8 +87,11 @@ public:
 
 	virtual void SetupTest();
 	virtual void RunTest();
-	virtual void TeardownTest();
+	virtual void TeardownTest(bool bExitAfterTeardown = true);
+	virtual void TriggerExitAfterDelay();
 	virtual void Exit();
+
+	AGameModeBase* GetGameMode() const;
 	
 protected:
 	// ~Begin UGauntletTestController Interface
@@ -80,11 +111,9 @@ protected:
 	virtual void UnbindAllDelegates();
 
 private:
-	FString ExplicitTestID;
-
 	FString TraceChannels;
 	FString TestDatetime;
-	FString TestID;
+	FString TestName;
 	FString DeviceProfileOverride;
 	bool bRequestsFPSChart;
 	bool bRequestsInsightsTrace;
diff --git a/Engine/Plugins/Performance/AutomatedPerfTesting/Source/AutomatedPerfTesting/Public/AutomatedPerfTestProjectSettings.h b/Engine/Plugins/Performance/AutomatedPerfTesting/Source/AutomatedPerfTesting/Public/AutomatedPerfTestProjectSettings.h
new file mode 100644
--- /dev/null
+++ b/Engine/Plugins/Performance/AutomatedPerfTesting/Source/AutomatedPerfTesting/Public/AutomatedPerfTestProjectSettings.h
@@ -0,0 +1,29 @@
+// Copyright Epic Games, Inc. All Rights Reserved.
+#pragma once
+
+#include "UObject/Object.h"
+#include "Templates/SubclassOf.h"
+#include "Engine/EngineTypes.h"
+#include "Engine/DeveloperSettings.h"
+#include "AutomatedPerfTestProjectSettings.generated.h"
+
+
+UCLASS(BlueprintType, Config=Engine, DefaultConfig, DisplayName="Automated Performance Testing")
+class AUTOMATEDPERFTESTING_API UAutomatedPerfTestProjectSettings : public UDeveloperSettings
+{
+	GENERATED_BODY()
+
+public:
+	UAutomatedPerfTestProjectSettings(const FObjectInitializer&);
+
+	/** Gets the settings container name for the settings, either Project or Editor */
+	virtual FName GetContainerName() const override { return FName("Project"); }
+	/** Gets the category for the settings, some high level grouping like, Editor, Engine, Game...etc. */
+	virtual FName GetCategoryName() const override { return FName("Plugins"); }
+
+	/* How many seconds to wait before transition from Teardown to Exiting the test */
+	UPROPERTY(Config, EditAnywhere, BlueprintReadWrite, Category="Automated Performance Testing")
+	float TeardownToExitDelay = 5.0;
+};
+
+
diff --git a/Engine/Plugins/Performance/AutomatedPerfTesting/Source/AutomatedPerfTesting/Public/AutomatedSequencePerfTest.h b/Engine/Plugins/Performance/AutomatedPerfTesting/Source/AutomatedPerfTesting/Public/AutomatedSequencePerfTest.h
--- a/Engine/Plugins/Performance/AutomatedPerfTesting/Source/AutomatedPerfTesting/Public/AutomatedSequencePerfTest.h
+++ b/Engine/Plugins/Performance/AutomatedPerfTesting/Source/AutomatedPerfTesting/Public/AutomatedSequencePerfTest.h
@@ -2,14 +2,76 @@
 
 #pragma once
 
+#include "UObject/Object.h"
+#include "Templates/SubclassOf.h"
+#include "Engine/EngineTypes.h"
+#include "Engine/DeveloperSettings.h"
 #include "CoreMinimal.h"
 #include "AutomatedPerfTestControllerBase.h"
+#include "Camera/CameraComponent.h"
 #include "AutomatedSequencePerfTest.generated.h"
 
+
 class ULevelSequence;
 class ALevelSequenceActor;
 class ULevelSequencePlayer;
 
+
+USTRUCT(BlueprintType)
+struct FAutomatedPerfTestMapSequenceCombo
+{
+	GENERATED_BODY()
+
+	/**
+	 * Use this name to directly reference this map/sequence combo via BuildGraph or UAT with -AutomatedPerfTest.SequenceTest.MapSequenceComboName
+	 */
+	UPROPERTY(Config, EditAnywhere, BlueprintReadWrite, Category="Sequence Perf Test", meta=(AllowedClasses="/Script/Engine.World"))
+	FName ComboName;
+
+	UPROPERTY(Config, EditAnywhere, BlueprintReadWrite, Category="Sequence Perf Test", meta=(AllowedClasses="/Script/Engine.World"))
+	FSoftObjectPath Map;
+	
+	UPROPERTY(Config, EditAnywhere, BlueprintReadWrite, Category="Sequence Perf Test", meta=(AllowedClasses="/Script/LevelSequence.LevelSequence"))
+	FSoftObjectPath Sequence;
+
+	/*
+	 * The name of the alias of the game mode you can optionally override when opening the level
+	 */
+	UPROPERTY(Config, EditAnywhere, BlueprintReadWrite, Category="Sequence Perf Test", meta=(AllowedClasses="/Script/GameFramework.GameModeBase"))
+	FString GameModeOverride;
+	
+};
+
+
+UCLASS(BlueprintType, Config=Engine, DefaultConfig, DisplayName="Automated Performance Testing | Sequence")
+class AUTOMATEDPERFTESTING_API UAutomatedSequencePerfTestProjectSettings : public UDeveloperSettings
+{
+	GENERATED_BODY()
+
+public:
+	UAutomatedSequencePerfTestProjectSettings(const FObjectInitializer&);
+
+	/** Gets the settings container name for the settings, either Project or Editor */
+	virtual FName GetContainerName() const override { return FName("Project"); }
+	/** Gets the category for the settings, some high level grouping like, Editor, Engine, Game...etc. */
+	virtual FName GetCategoryName() const override { return FName("Plugins"); }
+
+	/*
+	 * When the project is run with a Sequence Perf Test, cycle through the input maps, and load and run the associated sequence
+	 * outputting separate profiling results for each map/sequence combo
+	 * Can be overridden via commandline with -AutomatedPerfTest.SequencePerfTest.TestName, which will only run the test with the matching name
+	 */
+	UPROPERTY(Config, EditAnywhere, BlueprintReadWrite, Category="Sequence Perf Test")
+	TArray<FAutomatedPerfTestMapSequenceCombo> MapsAndSequencesToTest;
+
+	UFUNCTION(BlueprintCallable, Category="Sequence Perf Test")
+	bool GetComboFromTestName(FName TestName, FAutomatedPerfTestMapSequenceCombo& FoundSequence) const;
+
+	/* How long to delay between setting up the sequence for the map before the sequence actually starts */
+	UPROPERTY(Config, EditAnywhere, BlueprintReadWrite, Category="Sequence Perf Test")
+	float SequenceStartDelay=5.0;
+};
+
 /**
  * 
  */
@@ -19,25 +81,42 @@ class AUTOMATEDPERFTESTING_API UAutomatedSequencePerfTest : public UAutomatedPer
 	GENERATED_BODY()
 
 public:
-	FString GetSequenceName() const;
-	
+	// ~Begin UAutomatedPerfTestControllerBase Interface
+	virtual FString GetTestID() override;
 	virtual void SetupTest() override;
-
+	
 	UFUNCTION()
 	virtual void RunTest() override;
+	
+	virtual void TeardownTest(bool bExitAfterTeardown = true) override;
+	virtual void Exit() override;
+	// ~End UAutomatedPerfTestControllerBase Interface
+	
+	// This function is called on world change to set up the map for the correct map/sequence combo
+    UFUNCTION()
+    void NextMap();
 
 	UFUNCTION()
-	virtual void TeardownTest() override;
-	virtual void Exit() override;
+	void OnSequenceFinished();
+	
+	UFUNCTION()
+	void OnCameraCut(UCameraComponent* CameraComponent);
+	
+	FString GetCameraCutID();
 	
 protected:
 	virtual void OnInit() override;
 	virtual void UnbindAllDelegates() override;
 
 private:
-	FString SequencePathName;
-	FSoftObjectPath SequenceSoftPath;
+	const UAutomatedSequencePerfTestProjectSettings* Settings;
+
+	FName SequenceTestName;
+	FAutomatedPerfTestMapSequenceCombo CurrentMapSequenceCombo;
+	TArray<FAutomatedPerfTestMapSequenceCombo*> MapSequenceCombos;
 
 	ALevelSequenceActor* SequenceActor;
 	ULevelSequencePlayer* SequencePlayer;
+
+	UCameraComponent* CurrentCamera;
 };
