[Build Health] Fix code to address warnings introduced in CL 35422117: `The address-of a field cannot be null in well-defined code.`

diff --git a/Engine/Plugins/Performance/AutomatedPerfTesting/Source/AutomatedPerfTesting/Private/AutomatedSequencePerfTest.cpp b/Engine/Plugins/Performance/AutomatedPerfTesting/Source/AutomatedPerfTesting/Private/AutomatedSequencePerfTest.cpp
--- a/Engine/Plugins/Performance/AutomatedPerfTesting/Source/AutomatedPerfTesting/Private/AutomatedSequencePerfTest.cpp
+++ b/Engine/Plugins/Performance/AutomatedPerfTesting/Source/AutomatedPerfTesting/Private/AutomatedSequencePerfTest.cpp
@@ -39,9 +39,9 @@ bool UAutomatedSequencePerfTestProjectSettings::GetComboFromTestName(FName TestN
  *****/
 void UAutomatedSequencePerfTest::SetupTest()
 {
-	if(&CurrentMapSequenceCombo != nullptr)
+	if(CurrentMapSequenceCombo.IsSet())
 	{
-		if(GetCurrentMap() == CurrentMapSequenceCombo.Map.GetAssetName())
+		if(GetCurrentMap() == CurrentMapSequenceCombo->Map.GetAssetName())
 		{
 			// don't even try to set up the test if we're not in the correct map
 			Super::SetupTest();
@@ -50,8 +50,8 @@ void UAutomatedSequencePerfTest::SetupTest()
 			if(UWorld* const World = GetWorld())
 			{
 				// load the sequence specified by the user
-				UE_LOG(LogAutomatedPerfTest, Log, TEXT("Loading sequence %s"), *CurrentMapSequenceCombo.Sequence.ToString());
-				ULevelSequence* TargetSequence = LoadObject<ULevelSequence>(NULL, *CurrentMapSequenceCombo.Sequence.ToString(), NULL, LOAD_None, NULL);
+				UE_LOG(LogAutomatedPerfTest, Log, TEXT("Loading sequence %s"), *CurrentMapSequenceCombo->Sequence.ToString());
+				ULevelSequence* TargetSequence = LoadObject<ULevelSequence>(NULL, *CurrentMapSequenceCombo->Sequence.ToString(), NULL, LOAD_None, NULL);
 				check(TargetSequence);
 		
 				UE_LOG(LogAutomatedPerfTest, Log, TEXT("World is valid, creating sequence player"));
@@ -86,7 +86,7 @@ void UAutomatedSequencePerfTest::SetupTest()
 		}
 		else
 		{
-			UE_LOG(LogAutomatedPerfTest, Log, TEXT("Current Map Name %s is not expected %s, calling NextMap."), *GetCurrentMap(), *CurrentMapSequenceCombo.Map.GetAssetName())
+			UE_LOG(LogAutomatedPerfTest, Log, TEXT("Current Map Name %s is not expected %s, calling NextMap."), *GetCurrentMap(), *CurrentMapSequenceCombo->Map.GetAssetName())
 			NextMap();
 		}
 	}
@@ -103,19 +103,19 @@ void UAutomatedSequencePerfTest::NextMap()
 
 	if(MapSequenceCombos.Num() > 0)
 	{
-		CurrentMapSequenceCombo = *MapSequenceCombos.Pop();
-		UE_LOG(LogAutomatedPerfTest, Log, TEXT("Setting up test for Map/Sequence combo %s"), *CurrentMapSequenceCombo.ComboName.ToString())
+		CurrentMapSequenceCombo = MapSequenceCombos.Pop();
+		UE_LOG(LogAutomatedPerfTest, Log, TEXT("Setting up test for Map/Sequence combo %s"), *CurrentMapSequenceCombo->ComboName.ToString())
 
 		// no need to prepend this with a ? since OpenLevel handles that part for us
 		FString OptionsString;
-		if(!CurrentMapSequenceCombo.GameModeOverride.IsEmpty())
+		if(!CurrentMapSequenceCombo->GameModeOverride.IsEmpty())
 		{
-			UE_LOG(LogAutomatedPerfTest, Log, TEXT("Game Mode overridden to %s"), *CurrentMapSequenceCombo.GameModeOverride)
-			OptionsString += "game=" + CurrentMapSequenceCombo.GameModeOverride;
+			UE_LOG(LogAutomatedPerfTest, Log, TEXT("Game Mode overridden to %s"), *CurrentMapSequenceCombo->GameModeOverride)
+			OptionsString += "game=" + CurrentMapSequenceCombo->GameModeOverride;
 		}
 		
-		UE_LOG(LogAutomatedPerfTest, Log, TEXT("Opening map %s%s"), *CurrentMapSequenceCombo.Map.GetAssetName(), *OptionsString);
-		UGameplayStatics::OpenLevel(AutomatedPerfTest::FindCurrentWorld(), *CurrentMapSequenceCombo.Map.GetAssetName(), true, OptionsString);
+		UE_LOG(LogAutomatedPerfTest, Log, TEXT("Opening map %s%s"), *CurrentMapSequenceCombo->Map.GetAssetName(), *OptionsString);
+		UGameplayStatics::OpenLevel(AutomatedPerfTest::FindCurrentWorld(), *CurrentMapSequenceCombo->Map.GetAssetName(), true, OptionsString);
 	}
 	else
 	{
@@ -220,7 +220,7 @@ void UAutomatedSequencePerfTest::OnCameraCut(UCameraComponent* CameraComponent)
 
 FString UAutomatedSequencePerfTest::GetTestID()
 {
-	return &CurrentMapSequenceCombo != nullptr ? Super::GetTestID() + "_" + CurrentMapSequenceCombo.ComboName.ToString() : Super::GetTestID();
+	return CurrentMapSequenceCombo.IsSet() ? Super::GetTestID() + "_" + CurrentMapSequenceCombo->ComboName.ToString() : Super::GetTestID();
 }
 
 FString UAutomatedSequencePerfTest::GetCameraCutID()
@@ -244,9 +244,10 @@ void UAutomatedSequencePerfTest::OnInit()
 	// if an explicit map/sequence name was set from commandline, use this to override the test
 	if (FParse::Value(FCommandLine::Get(), TEXT("AutomatedPerfTest.SequencePerfTest.MapSequenceName="), SequenceTestName))
 	{
-		if(Settings->GetComboFromTestName(SequenceTestName, CurrentMapSequenceCombo))
+		FAutomatedPerfTestMapSequenceCombo MapSequenceCombo;
+		if(Settings->GetComboFromTestName(SequenceTestName, MapSequenceCombo))
 		{
-			MapSequenceCombos.Add(&CurrentMapSequenceCombo);
+			MapSequenceCombos.Add(MapSequenceCombo);
 		}
 		else
 		{
@@ -256,9 +257,9 @@ void UAutomatedSequencePerfTest::OnInit()
 	// otherwise, use all the maps defined in project settings
 	else
 	{
-		for(FAutomatedPerfTestMapSequenceCombo MapSequenceCombo : Settings->MapsAndSequencesToTest)
+		for(const FAutomatedPerfTestMapSequenceCombo& MapSequenceCombo : Settings->MapsAndSequencesToTest)
 		{
-			MapSequenceCombos.Add(&MapSequenceCombo);
+			MapSequenceCombos.Add(MapSequenceCombo);
 		}
 	}
 	
diff --git a/Engine/Plugins/Performance/AutomatedPerfTesting/Source/AutomatedPerfTesting/Public/AutomatedSequencePerfTest.h b/Engine/Plugins/Performance/AutomatedPerfTesting/Source/AutomatedPerfTesting/Public/AutomatedSequencePerfTest.h
--- a/Engine/Plugins/Performance/AutomatedPerfTesting/Source/AutomatedPerfTesting/Public/AutomatedSequencePerfTest.h
+++ b/Engine/Plugins/Performance/AutomatedPerfTesting/Source/AutomatedPerfTesting/Public/AutomatedSequencePerfTest.h
@@ -112,8 +112,8 @@ private:
 	const UAutomatedSequencePerfTestProjectSettings* Settings;
 
 	FName SequenceTestName;
-	FAutomatedPerfTestMapSequenceCombo CurrentMapSequenceCombo;
-	TArray<FAutomatedPerfTestMapSequenceCombo*> MapSequenceCombos;
+	TOptional<FAutomatedPerfTestMapSequenceCombo> CurrentMapSequenceCombo;
+	TArray<FAutomatedPerfTestMapSequenceCombo> MapSequenceCombos;
 
 	ALevelSequenceActor* SequenceActor;
 	ULevelSequencePlayer* SequencePlayer;
