Automated Perf Testing - Adding new tests and CSV Output granularity
+ Adding options to each test to determine the granularity at which CSV files are output.
+ Adding Material Perf test, which uses an input list of materials from your project settings and displays those materials in turn on a full-screen quad in a barebones scene, and collects performance data from each material
+ Adding Static Camera perf test base class and example implementation. This test is built in such a way that you can customize the method by which static cameras are gathered for the test.
++ The example usage looks for actors matching the AutomatedPerfTestStaticCamera class, and uses those as the set cameras.
* Changing how the DataSourceName is constructed by the AutomatedPerfTestNodes so that by default it uses the activity of the test, and can be overridden via commandline
* Fixing bug with Insights trace caching

#virtualized

diff --git a/Engine/Plugins/Performance/AutomatedPerfTesting/AutomatedPerfTesting.uplugin b/Engine/Plugins/Performance/AutomatedPerfTesting/AutomatedPerfTesting.uplugin
--- a/Engine/Plugins/Performance/AutomatedPerfTesting/AutomatedPerfTesting.uplugin
+++ b/Engine/Plugins/Performance/AutomatedPerfTesting/AutomatedPerfTesting.uplugin
@@ -1,9 +1,9 @@
 {
 	"FileVersion": 3,
 	"Version": 1,
-	"VersionName": "1.0",
+	"VersionName": "0.1",
 	"FriendlyName": "AutomatedPerfTesting",
-	"Description": "",
+	"Description": "This plugin provides Gauntlet Test Controllers to facilitate automatic performance testing.",
 	"Category": "Testing",
 	"CreatedBy": "",
 	"CreatedByURL": "",
diff --git a/Engine/Plugins/Performance/AutomatedPerfTesting/Build/Inc/AutomatedMaterialPerfTestMacro.xml b/Engine/Plugins/Performance/AutomatedPerfTesting/Build/Inc/AutomatedMaterialPerfTestMacro.xml
new file mode 100644
--- /dev/null
+++ b/Engine/Plugins/Performance/AutomatedPerfTesting/Build/Inc/AutomatedMaterialPerfTestMacro.xml
@@ -0,0 +1,30 @@
+<?xml version='1.0' ?>
+<BuildGraph xmlns="http://www.epicgames.com/BuildGraph" xmlns:xsi="http://www.w3.org/2001/XMLSchema-instance" xsi:schemaLocation="http://www.epicgames.com/BuildGraph ../../../../../../../Engine/Build/Graph/Schema.xsd" >
+	
+	<Macro Name="AutomatedMaterialPerfTest" Arguments="TestName;AdditionalArguments">
+		<Property Name="TestNameTitle" Value="$(TestName)"/>
+		<Expand Name="SanitizeTestNameTitle"/>
+		
+		<!-- Run and report on the test-->
+		<Do If="'$(AutomationReportOutputDirectory)' != ''">
+			<Property Name="OutputFragmentPath" Value="$(TargetPlatform)/$(TestConfiguration)/$(TestNameTitle)" />
+			<Property Name="OutputReportPath" Value="$(AutomationReportOutputDirectory)/$(OutputFragmentPath)" />
+			<Property Name="GauntletReportArgs" Value="-logdir=&quot;$(OutputReportPath)&quot; -artifactname=Gauntlet -ReportExportPath=&quot;$(OutputReportPath)&quot;" />
+			<Property Name="GauntletReportArgs" Value="$(GauntletReportArgs) -ReportURL=&quot;$(AutomationReportUrl)/$(OutputFragmentPath)/&quot;" If="'$(AutomationReportUrl)' != ''" />
+		</Do>
+		<Property Name="TestNodeName" Value="$(PreNodeName)$(TargetPlatform) $(TestConfiguration) $(TestNameTitle) MaterialPerformanceTest" />
+		<Property Name="NodeAnnotations" Value=""/>
+		<Property Name="NodeAnnotations" Value="Workflow=$(AnnotationsTarget)" If="ContainsItem('$(AnnotationsPlatforms)', '$(TargetPlatform)', ';') and ContainsItem('$(AnnotationsTests)', '$(TestName)', ';')"/>
+
+		<!-- Run the test. No 'Requires' as we run on the same machine locally and use the publish folder otherwise -->
+		<Node Name="$(TestNodeName)" Requires="$(TestPrerequisites)" Annotations="$(NodeAnnotations)">
+			<Command Name="RunUnreal" Arguments="-project=$(TargetProject) -test=&quot;$(TestName)&quot; -platform=$(TargetPlatform) -configuration=$(TestConfiguration) -target=$(TargetName) $(GauntletBuildArg) $(GauntletTestArgs) $(GauntletReportArgs) $(AdditionalArguments)"/>
+		</Node>
+
+		<!-- Add it to our list -->
+		<Do If="!ContainsItem('$(TargetPlatformsToSkipTests)', '$(TargetPlatform)', '+')">
+			<Property Name="AllPlatformTestNodes" Value="$(AllPlatformTestNodes);$(TestNodeName)"/>
+		</Do>
+	</Macro>
+	
+</BuildGraph>
diff --git a/Engine/Plugins/Performance/AutomatedPerfTesting/Build/Inc/AutomatedSequencePerfTestMacro.xml b/Engine/Plugins/Performance/AutomatedPerfTesting/Build/Inc/AutomatedSequencePerfTestMacro.xml
--- a/Engine/Plugins/Performance/AutomatedPerfTesting/Build/Inc/AutomatedSequencePerfTestMacro.xml
+++ b/Engine/Plugins/Performance/AutomatedPerfTesting/Build/Inc/AutomatedSequencePerfTestMacro.xml
@@ -1,10 +1,12 @@
 <?xml version='1.0' ?>
 <BuildGraph xmlns="http://www.epicgames.com/BuildGraph" xmlns:xsi="http://www.w3.org/2001/XMLSchema-instance" xsi:schemaLocation="http://www.epicgames.com/BuildGraph ../../../../../../../Engine/Build/Graph/Schema.xsd" >
 
+	<Property Name="MapSequenceComboNameArg" Value="" />
+	
 	<Macro Name="AutomatedSequencePerfTestNode">
 		<!-- Run the test. No 'Requires' as we run on the same machine locally and use the publish folder otherwise -->
 		<Node Name="$(TestNodeName)" Requires="$(TestPrerequisites)" Annotations="$(NodeAnnotations)">
-			<Command Name="RunUnreal" Arguments="-project=$(TargetProject) -test=&quot;$(TestName)&quot; -platform=$(TargetPlatform) -configuration=$(TestConfiguration) -target=$(TargetName) $(GauntletBuildArg) $(GauntletTestArgs) $(GauntletReportArgs) $(AdditionalArguments)"/>
+			<Command Name="RunUnreal" Arguments="-project=$(TargetProject) -test=&quot;$(TestName)&quot; -platform=$(TargetPlatform) -configuration=$(TestConfiguration) -target=$(TargetName) $(GauntletBuildArg) $(GauntletTestArgs) $(GauntletReportArgs) $(AdditionalArguments) $(MapSequenceComboNameArg)"/>
 		</Node>
 	</Macro>
 	
@@ -24,7 +26,7 @@
 			<Property Name="NodeAnnotations" Value=""/>
 			<Property Name="NodeAnnotations" Value="Workflow=$(AnnotationsTarget)" If="ContainsItem('$(AnnotationsPlatforms)', '$(TargetPlatform)', ';') and ContainsItem('$(AnnotationsTests)', '$(TestName)', ';')"/>
 			
-			<Property Name="GauntletTestArgs" Value="$(GauntletTestArgs) -AutomatedPerfTest.SequencePerfTest.MapSequenceName=$(MapSequenceName)" />
+			<Property Name="MapSequenceComboNameArg" Value="-AutomatedPerfTest.SequencePerfTest.MapSequenceName=$(MapSequenceName)" />
 			
 			<Expand Name="AutomatedSequencePerfTestNode" />
 
diff --git a/Engine/Plugins/Performance/AutomatedPerfTesting/Build/Inc/AutomatedStaticCameraPerfTestMacro.xml b/Engine/Plugins/Performance/AutomatedPerfTesting/Build/Inc/AutomatedStaticCameraPerfTestMacro.xml
new file mode 100644
--- /dev/null
+++ b/Engine/Plugins/Performance/AutomatedPerfTesting/Build/Inc/AutomatedStaticCameraPerfTestMacro.xml
@@ -0,0 +1,38 @@
+<?xml version='1.0' ?>
+<BuildGraph xmlns="http://www.epicgames.com/BuildGraph" xmlns:xsi="http://www.w3.org/2001/XMLSchema-instance" xsi:schemaLocation="http://www.epicgames.com/BuildGraph ../../../../../../../Engine/Build/Graph/Schema.xsd" >
+
+	<Macro Name="AutomatedStaticCameraPerfTestNode">
+		<!-- Run the test. No 'Requires' as we run on the same machine locally and use the publish folder otherwise -->
+		<Node Name="$(TestNodeName)" Requires="$(TestPrerequisites)" Annotations="$(NodeAnnotations)">
+			<Command Name="RunUnreal" Arguments="-project=$(TargetProject) -test=&quot;$(TestName)&quot; -platform=$(TargetPlatform) -configuration=$(TestConfiguration) -target=$(TargetName) $(GauntletBuildArg) $(GauntletTestArgs) $(GauntletReportArgs) $(AdditionalArguments) $(MapNameArgument)"/>
+		</Node>
+	</Macro>
+	
+	<Macro Name="AutomatedStaticCameraPerfTest" Arguments="TestName;MapNames;AdditionalArguments">
+		<ForEach Name="MapName" Values="$(MapNames)" Separator="+">
+			<Property Name="TestNameTitle" Value="$(TestName)"/>
+			<Expand Name="SanitizeTestNameTitle"/>
+			
+			<!-- Run and report on the test-->
+			<Do If="'$(AutomationReportOutputDirectory)' != ''">
+				<Property Name="OutputFragmentPath" Value="$(TargetPlatform)/$(TestConfiguration)/$(TestNameTitle)" />
+				<Property Name="OutputReportPath" Value="$(AutomationReportOutputDirectory)/$(OutputFragmentPath)" />
+				<Property Name="GauntletReportArgs" Value="-logdir=&quot;$(OutputReportPath)&quot; -artifactname=Gauntlet -ReportExportPath=&quot;$(OutputReportPath)&quot;" />
+				<Property Name="GauntletReportArgs" Value="$(GauntletReportArgs) -ReportURL=&quot;$(AutomationReportUrl)/$(OutputFragmentPath)/&quot;" If="'$(AutomationReportUrl)' != ''" />
+			</Do>
+			<Property Name="TestNodeName" Value="$(PreNodeName)$(TargetPlatform) $(TestConfiguration) $(TestNameTitle) $(MapName)" />
+			<Property Name="NodeAnnotations" Value=""/>
+			<Property Name="NodeAnnotations" Value="Workflow=$(AnnotationsTarget)" If="ContainsItem('$(AnnotationsPlatforms)', '$(TargetPlatform)', ';') and ContainsItem('$(AnnotationsTests)', '$(TestName)', ';')"/>
+			
+			<Property Name="MapNameArgument" Value="-AutomatedPerfTest.StaticCameraPerfTest.MapName=$(MapName)" />
+			
+			<Expand Name="AutomatedStaticCameraPerfTestNode" />
+
+			<!-- Add it to our list -->
+			<Do If="!ContainsItem('$(TargetPlatformsToSkipTests)', '$(TargetPlatform)', '+')">
+				<Property Name="AllPlatformTestNodes" Value="$(AllPlatformTestNodes);$(TestNodeName)"/>
+			</Do>
+		</ForEach>
+	</Macro>
+	
+</BuildGraph>
diff --git a/Engine/Plugins/Performance/AutomatedPerfTesting/Build/Scripts/AutomatedPerfTestConfig.cs b/Engine/Plugins/Performance/AutomatedPerfTesting/Build/Scripts/AutomatedPerfTestConfig.cs
--- a/Engine/Plugins/Performance/AutomatedPerfTesting/Build/Scripts/AutomatedPerfTestConfig.cs
+++ b/Engine/Plugins/Performance/AutomatedPerfTesting/Build/Scripts/AutomatedPerfTestConfig.cs
@@ -6,16 +6,26 @@ using System.IO;
 using System.Linq;
 using EpicGame;
 using Gauntlet;
+using Microsoft.Extensions.Logging;
+using Log = EpicGames.Core.Log;
 
 namespace AutomatedPerfTest
 {
 	public class AutomatedPerfTestConfigBase : UnrealTestConfiguration
 	{
 		/// <summary>
-		/// Name of the project (TODO: parse from the build metadata? What about github builds?)
+		/// Used to override the test controller portion of the data source name
+		/// Will be appended to Automation.ProjectName to construct the full data source name for the test
+		/// can be overridden with AutomatedPerfTest.DataSourceNameOverride 
 		/// </summary>
-		[AutoParamWithNames("AutomatedPerfTest.DataSourceName")]
-		public string DataSourceName = "";
+		[AutoParamWithNames("AutomatedPerfTest.DataSourceTypeOverride")]
+		public string DataSourceTypeOverride = "";
+		
+		/// <summary>
+		/// Fully override the data source name
+		/// </summary>
+		[AutoParamWithNames("AutomatedPerfTest.DataSourceNameOverride")]
+		public string DataSourceNameOverride = "";
 
 		/// <summary>
 		/// Name of the test, useful for identifying it later
@@ -77,6 +87,35 @@ namespace AutomatedPerfTest
 		/// </summary>
 		[AutoParamWithNames(true, "AutomatedPerfTest.UseShippingInsights")]
 		public bool UseShippingInsights;
+		
+		public string DataSourceName;
+
+		/// <summary>
+		/// Call this in the test node's GetConfiguration function to set the DataSourceName used for later data processing
+		/// </summary>
+		/// <param name="ProjectName">Name of the project</param>
+		/// <param name="DataSourceType">Which test controller is being used to generate the data</param>
+		/// <returns> the properly formatted data source name </returns>
+		public string GetDataSourceName(string ProjectName, string DataSourceType="")
+		{
+			// if the name has been fully overridden by the calling process, just use that
+			if (!string.IsNullOrEmpty(DataSourceNameOverride))
+			{
+				return DataSourceNameOverride;
+			}
+			
+			// otherwise, if the data source type has been override, use either that or the one passed into this function
+
+			if (string.IsNullOrEmpty(DataSourceType) && string.IsNullOrEmpty(DataSourceTypeOverride))
+			{
+				Log.Logger.LogError("No DataSourceType or DataSourceTypeOverride has been provided.");
+				return null;
+			}
+			
+			string dataSourceType = string.IsNullOrEmpty(DataSourceTypeOverride) ? DataSourceType : DataSourceTypeOverride;
+			
+			return $"Automation.{ProjectName}.{dataSourceType}";
+		}
 	}
 
 	public class AutomatedSequencePerfTestConfig : AutomatedPerfTestConfigBase
@@ -87,4 +126,17 @@ namespace AutomatedPerfTest
 		[AutoParamWithNames("", "AutomatedPerfTest.SequencePerfTest.MapSequenceName")]
 		public string MapSequenceComboName;
 	}
+
+	public class AutomatedStaticCameraPerfTestConfig : AutomatedPerfTestConfigBase
+	{
+		/// <summary>
+		/// Which map to run the test on
+		/// </summary>
+		[AutoParamWithNames("", "AutomatedPerfTest.StaticCameraPerfTest.MapName")]
+		public string MapName;
+	}
+	
+	public class AutomatedMaterialPerfTestConfig : AutomatedPerfTestConfigBase
+	{
+	}
 }
diff --git a/Engine/Plugins/Performance/AutomatedPerfTesting/Build/Scripts/AutomatedPerfTestNode.cs b/Engine/Plugins/Performance/AutomatedPerfTesting/Build/Scripts/AutomatedPerfTestNode.cs
--- a/Engine/Plugins/Performance/AutomatedPerfTesting/Build/Scripts/AutomatedPerfTestNode.cs
+++ b/Engine/Plugins/Performance/AutomatedPerfTesting/Build/Scripts/AutomatedPerfTestNode.cs
@@ -554,11 +554,6 @@ namespace AutomatedPerfTest
 					ClientRole.CommandLineParams.AddOrAppendParamValue("AutomatedPerfTest.TraceChannels", Config.TraceChannels);
 				}
 			}
-			
-			if (string.IsNullOrEmpty(Config.DataSourceName))
-			{
-				Config.DataSourceName = string.Format("Automation.{0}.Performance", Context.BuildInfo.ProjectName);
-			}
 
 			if (Config.DoFPSChart)
 			{
@@ -638,7 +633,7 @@ namespace AutomatedPerfTest
 				// just try the copy over, and log a failure, but don't bail out of the test.
 				try
 				{
-					InternalUtils.SafeCreateDirectory(PerfCachePath, true);
+					InternalUtils.SafeCreateDirectory(Path.GetDirectoryName(PerfCachePath), true);
 					File.Copy(SortedTraces[0], PerfCachePath);
 				}
 				catch (Exception e)
@@ -677,6 +672,8 @@ namespace AutomatedPerfTest
 		public override TConfigClass GetConfiguration()
 		{
 			TConfigClass Config = base.GetConfiguration();
+
+			Config.DataSourceName = Config.GetDataSourceName(Context.BuildInfo.ProjectName, "Sequence");
 			
 			// extend the role(s) that we initialized in the base class
 			if (Config.GetRequiredRoles(UnrealTargetRole.Client).Any())
@@ -698,4 +695,72 @@ namespace AutomatedPerfTest
 			return Config;
 		}
 	}
+
+	/// <summary>
+	/// Implementation of a Gauntlet TestNode for AutomatedPerfTest plugin
+	/// </summary>
+	/// <typeparam name="TConfigClass"></typeparam>
+	public abstract class AutomatedStaticCameraPerfTestNode<TConfigClass> : AutomatedPerfTestNode<TConfigClass>
+		where TConfigClass : AutomatedStaticCameraPerfTestConfig, new()
+	{
+		public AutomatedStaticCameraPerfTestNode(UnrealTestContext InContext) : base(InContext)
+		{
+		}
+
+		public override TConfigClass GetConfiguration()
+		{
+			TConfigClass Config = base.GetConfiguration();
+			
+			Config.DataSourceName = Config.GetDataSourceName(Context.BuildInfo.ProjectName, "StaticCamera");
+			
+			// extend the role(s) that we initialized in the base class
+			if (Config.GetRequiredRoles(UnrealTargetRole.Client).Any())
+			{
+				foreach(UnrealTestRole ClientRole in Config.GetRequiredRoles(UnrealTargetRole.Client))
+				{
+					ClientRole.Controllers.Add("AutomatedPlacedStaticCameraPerfTest");
+					
+					// if a specific MapName was defined in the commandline to UAT, then add that to the commandline for the role
+					if (!string.IsNullOrEmpty(Config.MapName))
+					{
+						// use add Unique, since there should only ever be one of these specified
+						ClientRole.CommandLineParams.AddUnique($"AutomatedPerfTest.StaticCameraPerfTest.MapName",
+							Config.MapName);
+					}
+				}
+			}
+
+			return Config;
+		}
+	}
+	
+	/// <summary>
+	/// Implementation of a Gauntlet TestNode for AutomatedPerfTest plugin
+	/// </summary>
+	/// <typeparam name="TConfigClass"></typeparam>
+	public abstract class AutomatedMaterialPerfTestNode<TConfigClass> : AutomatedPerfTestNode<TConfigClass>
+		where TConfigClass : AutomatedMaterialPerfTestConfig, new()
+	{
+		public AutomatedMaterialPerfTestNode(UnrealTestContext InContext) : base(InContext)
+		{
+		}
+
+		public override TConfigClass GetConfiguration()
+		{
+			TConfigClass Config = base.GetConfiguration();
+			
+			Config.DataSourceName = Config.GetDataSourceName(Context.BuildInfo.ProjectName, "Material");
+			
+			// extend the role(s) that we initialized in the base class
+			if (Config.GetRequiredRoles(UnrealTargetRole.Client).Any())
+			{
+				foreach(UnrealTestRole ClientRole in Config.GetRequiredRoles(UnrealTargetRole.Client))
+				{
+					ClientRole.Controllers.Add("AutomatedMaterialPerfTest");
+				}
+			}
+
+			return Config;
+		}
+	}
 }
diff --git a/Engine/Plugins/Performance/AutomatedPerfTesting/Source/AutomatedPerfTesting/Private/AutomatedMaterialPerfTest.cpp b/Engine/Plugins/Performance/AutomatedPerfTesting/Source/AutomatedPerfTesting/Private/AutomatedMaterialPerfTest.cpp
new file mode 100644
--- /dev/null
+++ b/Engine/Plugins/Performance/AutomatedPerfTesting/Source/AutomatedPerfTesting/Private/AutomatedMaterialPerfTest.cpp
@@ -0,0 +1,257 @@
+﻿// Copyright Epic Games, Inc. All Rights Reserved.
+
+#include "AutomatedMaterialPerfTest.h"
+#include "TimerManager.h"
+#include "Kismet/GameplayStatics.h"
+#include "AutomatedPerfTesting.h"
+#include "Camera/CameraActor.h"
+#include "Camera/CameraComponent.h"
+#include "Engine/StaticMeshActor.h"
+#include "Engine/StaticMesh.h"
+#include "Engine/Player.h"
+#include "GameFramework/Pawn.h"
+#include "Materials/MaterialInterface.h"
+
+UAutomatedMaterialPerfTestProjectSettings::UAutomatedMaterialPerfTestProjectSettings(const FObjectInitializer& Initializer)
+	: Super(Initializer)
+	, bCaptureScreenshots(true)
+	, WarmUpTime(5.0)
+	, SoakTime(5.0)
+	, CooldownTime(1.0)
+	, MaterialPerformanceTestMap(FSoftObjectPath("/AutomatedPerfTesting/Tests/Materials/AutomatedMaterialPerfTestDefaultMap.AutomatedMaterialPerfTestDefaultMap"))
+	, CameraProjectionMode(ECameraProjectionMode::Type::Orthographic)
+	, PlateDistanceFromCamera(512.0)
+	, MaterialPlate(FSoftObjectPath("/AutomatedPerfTesting/Tests/Materials/SM_AutomatedMaterialPerfTestDefaultPlate.SM_AutomatedMaterialPerfTestDefaultPlate"))
+	, CSVOutputMode(EAutomatedPerfTestCSVOutputMode::Granular)
+{
+}
+
+void UAutomatedMaterialPerfTest::SetupTest()
+{
+	// load up into the map defined in project settings
+	if(GetCurrentMap() == Settings->MaterialPerformanceTestMap.GetAssetName())
+	{
+		Super::SetupTest();
+
+		UE_LOG(LogAutomatedPerfTest, Log, TEXT("UAutomatedMaterialPerfTest::SetupTest"))
+		// make sure the world exists, then create a sequence player
+		if(UWorld* const World = GetWorld())
+		{
+			// Also load and verify the material plate mesh
+			UE_LOG(LogAutomatedPerfTest, Log, TEXT("Loading material plate mesh: %s"), *Settings->MaterialPlate.ToString())
+			UStaticMesh* LoadedMaterialPlateMesh = LoadObject<UStaticMesh>(NULL, *Settings->MaterialPlate.ToString(), NULL, LOAD_None, NULL);
+			check(LoadedMaterialPlateMesh);
+
+			// reset the material index;
+			CurrentMaterialIndex = -1;
+			
+			// hide the pawn so it doesn't interfere with screenshots
+			GetFirstPlayerController()->GetPawn()->SetHidden(true);
+			
+			// then spawn the camera into the world and set it up 
+			Camera = GetWorld()->SpawnActor<ACameraActor>();
+			Camera->GetCameraComponent()->SetProjectionMode(Settings->CameraProjectionMode);
+			Camera->GetCameraComponent()->SetOrthoWidth(Settings->PlateDistanceFromCamera);
+
+			// and spawn the material plate into the world, and move it PlateDistanceFromCamera away down the X axis
+			MaterialPlate = GetWorld()->SpawnActor<AStaticMeshActor>();
+			MaterialPlate->SetMobility(EComponentMobility::Type::Movable);
+			MaterialPlate->GetStaticMeshComponent()->SetStaticMesh(LoadedMaterialPlateMesh);
+			MaterialPlate->SetActorLocation(FVector(Settings->PlateDistanceFromCamera, 0.0, 0.0));
+			float Scale = Settings->PlateDistanceFromCamera / (LoadedMaterialPlateMesh->GetBounds().BoxExtent.Y);
+
+			UE_LOG(LogAutomatedPerfTest, Verbose, TEXT("SizeY = %f, Scale = %f"), LoadedMaterialPlateMesh->GetBounds().BoxExtent.Y, Scale);
+			
+			MaterialPlate->SetActorScale3D(FVector(1.0, Scale, Scale));
+			
+			GetFirstPlayerController()->SetViewTarget(Camera);
+			
+        	// delay for WarmUpDelay, and call RunTest
+        	FTimerHandle UnusedHandle;
+        	GetWorld()->GetTimerManager().SetTimer(UnusedHandle, this, &UAutomatedMaterialPerfTest::RunTest, 1.0, false, Settings->WarmUpTime);
+		}
+		// if we have an invalid world, we can't run the test, so we should bail out
+		else
+		{
+			UE_LOG(LogAutomatedPerfTest, Error, TEXT("Invalid World when starting UAutomatedMaterialPerfTest, exiting..."));
+			EndAutomatedPerfTest(1);
+		}
+	}
+	else
+	{
+		UE_LOG(LogAutomatedPerfTest, Log, TEXT("Current Map Name %s is not the expected %s, loading the material performance test map"), *GetCurrentMap(), *Settings->MaterialPerformanceTestMap.GetAssetName())
+		OpenMaterialPerformanceTestMap();
+	}
+}
+
+void UAutomatedMaterialPerfTest::RunTest()
+{
+	Super::RunTest();
+
+	UE_LOG(LogAutomatedPerfTest, Log, TEXT("UAutomatedMaterialPerfTest::RunTest"))
+	
+	FTimerHandle UnusedHandle;
+	GetWorld()->GetTimerManager().SetTimer(UnusedHandle, this, &UAutomatedMaterialPerfTest::SetUpNextMaterial, 1.0, false, Settings->WarmUpTime);
+}
+
+FString UAutomatedMaterialPerfTest::GetTestID()
+{
+	return Super::GetTestID() + "_Materials";
+}
+
+void UAutomatedMaterialPerfTest::SetUpNextMaterial()
+{
+	CurrentMaterialIndex += 1;
+	
+	if(CurrentMaterialIndex >= Settings->MaterialsToTest.Num())
+	{
+		UE_LOG(LogAutomatedPerfTest, Log, TEXT("No more materials left to test, moving to teardown."))
+		TeardownTest();
+		return;
+	}
+	
+	// load the next material
+	FSoftObjectPath MaterialSoftObjectPath = Settings->MaterialsToTest[CurrentMaterialIndex];
+	
+	UE_LOG(LogAutomatedPerfTest, Log, TEXT("Loading material: %s"), *MaterialSoftObjectPath.ToString())
+	CurrentMaterial = LoadObject<UMaterialInterface>(NULL, *MaterialSoftObjectPath.ToString(), NULL, LOAD_None, NULL);
+	check(CurrentMaterial);
+
+	UE_LOG(LogAutomatedPerfTest, Log, TEXT("Applying material: %s"), *CurrentMaterial->GetName());
+	
+	MaterialPlate->GetStaticMeshComponent()->SetMaterial(0, CurrentMaterial);
+	
+	FTimerHandle UnusedHandle;
+	GetWorld()->GetTimerManager().SetTimer(UnusedHandle, this, &UAutomatedMaterialPerfTest::EvaluateMaterial, 1.0, false, Settings->WarmUpTime);
+}
+
+void UAutomatedMaterialPerfTest::EvaluateMaterial()
+{
+	MarkMaterialStart();
+
+	FTimerHandle UnusedHandle;
+	GetWorld()->GetTimerManager().SetTimer(UnusedHandle, this, &UAutomatedMaterialPerfTest::FinishMaterialEvaluation, 1.0, false, Settings->SoakTime);
+}
+
+void UAutomatedMaterialPerfTest::FinishMaterialEvaluation()
+{
+	MarkMaterialEnd();
+	
+	if(Settings->bCaptureScreenshots)
+	{
+		FTimerHandle UnusedHandle;
+		GetWorld()->GetTimerManager().SetTimer(UnusedHandle, this, &UAutomatedMaterialPerfTest::ScreenshotMaterial, 1.0, false, Settings->CooldownTime);		
+	}
+	else
+	{
+		FTimerHandle UnusedHandle;
+		GetWorld()->GetTimerManager().SetTimer(UnusedHandle, this, &UAutomatedMaterialPerfTest::SetUpNextMaterial, 1.0, false, Settings->CooldownTime);
+	}
+}
+
+void UAutomatedMaterialPerfTest::ScreenshotMaterial()
+{
+	TakeScreenshot(GetCurrentMaterialRegionName());
+
+	// start a timer to trigger the disk screenshot, since trace screenshots and disk screenshots can't happen in the same frame
+	FTimerHandle UnusedHandle;
+	GetWorld()->GetTimerManager().SetTimer(UnusedHandle, this, &UAutomatedMaterialPerfTest::SetUpNextMaterial, 1.0, false, .1);
+}
+
+UMaterialInterface* UAutomatedMaterialPerfTest::GetCurrentMaterial() const
+{
+	return CurrentMaterial;
+}
+
+FString UAutomatedMaterialPerfTest::GetCurrentMaterialRegionName()
+{
+	return GetCurrentMaterial()->GetName();
+}
+
+FString UAutomatedMaterialPerfTest::GetCurrentMaterialRegionFullName()
+{
+	return GetTestID() + "_" + GetCurrentMaterialRegionName();
+}
+
+void UAutomatedMaterialPerfTest::MarkMaterialStart()
+{
+	// safety check on the current material
+	if(CurrentMaterial)
+	{
+		if(RequestsInsightsTrace())
+		{
+			TRACE_BEGIN_REGION(*GetCurrentMaterialRegionName());
+		}
+#if CSV_PROFILER
+		if(RequestsCSVProfiler())
+		{
+			if(GetCSVOutputMode() == EAutomatedPerfTestCSVOutputMode::Granular)
+			{
+				TryStartCSVProfiler(GetCurrentMaterialRegionFullName());
+			}
+			CSV_EVENT(AutomatedPerfTest, TEXT("START_%s"), *GetCurrentMaterialRegionName())
+		}
+#endif
+	}
+}
+
+void UAutomatedMaterialPerfTest::MarkMaterialEnd()
+{
+	// safety check on the current material
+	if(CurrentMaterial)
+	{
+		if(RequestsInsightsTrace())
+		{
+			TRACE_END_REGION(*GetCurrentMaterialRegionName());
+		}
+#if CSV_PROFILER
+		if(RequestsCSVProfiler())
+		{
+			CSV_EVENT(AutomatedPerfTest, TEXT("END_%s"), *GetCurrentMaterialRegionName())
+			if(GetCSVOutputMode() == EAutomatedPerfTestCSVOutputMode::Granular)
+			{
+				TryStopCSVProfiler();
+			}
+		}
+#endif
+	}
+}
+
+void UAutomatedMaterialPerfTest::OnInit()
+{
+	Super::OnInit();
+	
+	UE_LOG(LogAutomatedPerfTest, Log, TEXT("UAutomatedMaterialPerfTest::OnInit"));
+
+	Settings = GetDefault<UAutomatedMaterialPerfTestProjectSettings>();
+
+	SetCSVOutputMode(Settings->CSVOutputMode);
+	
+	// early out if there aren't actually any materials
+	if(Settings->MaterialsToTest.Num() < 0)
+	{
+		UE_LOG(LogAutomatedPerfTest, Error, TEXT("No materials defined in the project's Automated Perf Test | Materials settings. Exiting test early."));
+		EndAutomatedPerfTest(1);
+	}
+}
+
+void UAutomatedMaterialPerfTest::UnbindAllDelegates()
+{
+	Super::UnbindAllDelegates();
+	
+	GetWorld()->GetTimerManager().ClearAllTimersForObject(this);
+}
+
+void UAutomatedMaterialPerfTest::OpenMaterialPerformanceTestMap() const
+{
+	// no need to prepend this with a ? since OpenLevel handles that part for us
+	FString OptionsString;
+	if(!Settings->GameModeOverride.IsEmpty())
+	{
+		UE_LOG(LogAutomatedPerfTest, Log, TEXT("Game Mode overridden to %s"), *Settings->GameModeOverride)
+		OptionsString += "game=" + Settings->GameModeOverride;
+	}
+	
+	UE_LOG(LogAutomatedPerfTest, Log, TEXT("Opening map %s%s"), *Settings->MaterialPerformanceTestMap.GetAssetName(), *OptionsString);
+	UGameplayStatics::OpenLevel(AutomatedPerfTest::FindCurrentWorld(), *Settings->MaterialPerformanceTestMap.GetAssetName(), true, OptionsString);
+}
diff --git a/Engine/Plugins/Performance/AutomatedPerfTesting/Source/AutomatedPerfTesting/Private/AutomatedPerfTestControllerBase.cpp b/Engine/Plugins/Performance/AutomatedPerfTesting/Source/AutomatedPerfTesting/Private/AutomatedPerfTestControllerBase.cpp
--- a/Engine/Plugins/Performance/AutomatedPerfTesting/Source/AutomatedPerfTesting/Private/AutomatedPerfTestControllerBase.cpp
+++ b/Engine/Plugins/Performance/AutomatedPerfTesting/Source/AutomatedPerfTesting/Private/AutomatedPerfTestControllerBase.cpp
@@ -2,8 +2,6 @@
 
 #include "AutomatedPerfTestControllerBase.h"
 
-#include <valarray>
-
 #include "AutomatedPerfTesting.h"
 #include "DeviceProfiles/DeviceProfileManager.h"
 #include "GameFramework/GameModeBase.h"
@@ -13,13 +11,14 @@
 #include "Misc/CommandLine.h"
 #include "ProfilingDebugging/TraceAuxiliary.h"
 #include "TimerManager.h"
-#include "Commandlets/Commandlet.h"
 #include "ProfilingDebugging/CsvProfiler.h"
 #include "Async/TaskGraphInterfaces.h"
 #include "VideoRecordingSystem.h"
 #include "PlatformFeatures.h"
 #include "Engine/Engine.h"
 #include "Misc/Paths.h"
+#include "ProfilingDebugging/TraceScreenshot.h"
+#include "UnrealClient.h"
 
 
 DEFINE_LOG_CATEGORY(LogAutomatedPerfTest)
@@ -95,8 +94,9 @@ UAutomatedPerfTestControllerBase::UAutomatedPerfTestControllerBase(const FObject
 	, bRequestsInsightsTrace(false)
 	, bRequestsCSVProfiler(false)
 	, bRequestsVideoCapture(false)
+    , CSVOutputMode(EAutomatedPerfTestCSVOutputMode::Single)
 {
-	// cache this off once, so that it's consistent across the board
+	// cache this off once, so that it's consistent throughout a session
 	TestDatetime = FDateTime::Now().ToString(TEXT("%Y%m%d-%H%M%S")); 
 }
 
@@ -172,15 +172,23 @@ bool UAutomatedPerfTestControllerBase::TryStopInsightsTrace()
 }
 
 bool UAutomatedPerfTestControllerBase::TryStartCSVProfiler()
+{
+	return TryStartCSVProfiler(GetTestID());
+}
+
+bool UAutomatedPerfTestControllerBase::TryStartCSVProfiler(FString CSVFileName)
 {
 #if CSV_PROFILER
 	if(FCsvProfiler* const CsvProfiler = FCsvProfiler::Get())
 	{
-		const FString CSVFilename = GetTestID() + ".csv";
+		if(!CSVFileName.EndsWith(".csv"))
+		{
+			CSVFileName += TEXT(".csv");
+		}
 
-		UE_LOG(LogAutomatedPerfTest, Log, TEXT("Attempting to start CSV Profile to file %s"), *CSVFilename);
+		UE_LOG(LogAutomatedPerfTest, Log, TEXT("Attempting to start CSV Profile to file %s"), *CSVFileName);
 		
-		CsvProfiler->BeginCapture(-1, FString(), CSVFilename);
+		CsvProfiler->BeginCapture(-1, FString(), CSVFileName);
 		CsvProfiler->SetDeviceProfileName(GetDeviceProfile());
 		
 		return CsvProfiler->IsCapturing();
@@ -300,7 +308,7 @@ void UAutomatedPerfTestControllerBase::SetupTest()
 		TryStartInsightsTrace();
 	}
 
-	if(RequestsCSVProfiler())
+	if(RequestsCSVProfiler() && CSVOutputMode ==  EAutomatedPerfTestCSVOutputMode::Single)
 	{
 		TryStartCSVProfiler();
 	}
@@ -323,6 +331,8 @@ void UAutomatedPerfTestControllerBase::SetupTest()
 	}
 
 	// Subclasses should implement their own transitions from SetupTest to RunTest depending on their needs
+
+	UE_LOG(LogAutomatedPerfTest, Log, TEXT("CSV Output Mode: %s"), *UEnum::GetValueAsString(CSVOutputMode))
 }
 
 void UAutomatedPerfTestControllerBase::RunTest()
@@ -335,6 +345,10 @@ void UAutomatedPerfTestControllerBase::RunTest()
 
 	if(RequestsCSVProfiler())
 	{
+		if(CSVOutputMode == EAutomatedPerfTestCSVOutputMode::Separate)
+		{
+			TryStartCSVProfiler();
+		}
 		CSV_EVENT(AutomatedPerfTest, TEXT("START"), *GetOverallRegionName())
 	}
 	
@@ -356,7 +370,10 @@ void UAutomatedPerfTestControllerBase::TeardownTest(bool bExitAfterTeardown)
 	if(RequestsCSVProfiler())
 	{
 		CSV_EVENT(AutomatedPerfTest, TEXT("END"), *GetOverallRegionName())
-		TryStopCSVProfiler();
+		if(CSVOutputMode == EAutomatedPerfTestCSVOutputMode::Separate)
+		{
+			TryStopCSVProfiler();
+		}
 	}
 
 	if(RequestsFPSChart())
@@ -397,6 +414,15 @@ void UAutomatedPerfTestControllerBase::Exit()
 		TryStopInsightsTrace();
 	}
 
+	if(RequestsCSVProfiler())
+	{
+		CSV_EVENT(AutomatedPerfTest, TEXT("END"), *GetOverallRegionName())
+		if(CSVOutputMode == EAutomatedPerfTestCSVOutputMode::Single)
+		{
+			TryStopCSVProfiler();
+		}
+	}
+
 	if(GameMode && GameMode->GetClass()->ImplementsInterface(UAutomatedPerfTestInterface::StaticClass()))
 	{
 		IAutomatedPerfTestInterface::Execute_Exit(GameMode);
@@ -429,6 +455,23 @@ AGameModeBase* UAutomatedPerfTestControllerBase::GetGameMode() const
 	return GameMode;
 }
 
+void UAutomatedPerfTestControllerBase::TakeScreenshot(FString ScreenshotName)
+{
+	if(RequestsInsightsTrace())
+	{
+		FTraceScreenshot::RequestScreenshot(ScreenshotName, false, LogAutomatedPerfTest);
+	}
+	else
+	{
+		FScreenshotRequest::RequestScreenshot(ScreenshotName, false, false);
+	}
+}
+
+void UAutomatedPerfTestControllerBase::SetCSVOutputMode(EAutomatedPerfTestCSVOutputMode NewOutputMode)
+{
+	CSVOutputMode = NewOutputMode;	
+}
+
 void UAutomatedPerfTestControllerBase::OnInit()
 {
 	Super::OnInit();
@@ -547,3 +590,8 @@ void UAutomatedPerfTestControllerBase::UnbindAllDelegates()
 		}
 	}
 }
+
+EAutomatedPerfTestCSVOutputMode UAutomatedPerfTestControllerBase::GetCSVOutputMode() const
+{
+	return CSVOutputMode;
+}
diff --git a/Engine/Plugins/Performance/AutomatedPerfTesting/Source/AutomatedPerfTesting/Private/AutomatedSequencePerfTest.cpp b/Engine/Plugins/Performance/AutomatedPerfTesting/Source/AutomatedPerfTesting/Private/AutomatedSequencePerfTest.cpp
--- a/Engine/Plugins/Performance/AutomatedPerfTesting/Source/AutomatedPerfTesting/Private/AutomatedSequencePerfTest.cpp
+++ b/Engine/Plugins/Performance/AutomatedPerfTesting/Source/AutomatedPerfTesting/Private/AutomatedSequencePerfTest.cpp
@@ -17,6 +17,7 @@
 
 UAutomatedSequencePerfTestProjectSettings::UAutomatedSequencePerfTestProjectSettings(const FObjectInitializer& Initializer)
 	: Super(Initializer)
+	, CSVOutputMode(EAutomatedPerfTestCSVOutputMode::Separate)
 {
 }
 
@@ -49,6 +50,9 @@ void UAutomatedSequencePerfTest::SetupTest()
 			// make sure the world exists, then create a sequence player
 			if(UWorld* const World = GetWorld())
 			{
+				// reset the camera cut number
+				NumCameraCuts = -1;
+				
 				// load the sequence specified by the user
 				UE_LOG(LogAutomatedPerfTest, Log, TEXT("Loading sequence %s"), *CurrentMapSequenceCombo->Sequence.ToString());
 				ULevelSequence* TargetSequence = LoadObject<ULevelSequence>(NULL, *CurrentMapSequenceCombo->Sequence.ToString(), NULL, LOAD_None, NULL);
@@ -136,8 +140,8 @@ void UAutomatedSequencePerfTest::RunTest()
 		UE_LOG(LogAutomatedPerfTest, Log, TEXT("RunTest::Valid Sequence Player, proceeding"));
 
 		// trigger a camera cut manually in order to start the region for the first camera cut
-		OnCameraCut(SequencePlayer->GetActiveCameraComponent());
         SequencePlayer->Play();
+		OnCameraCut(SequencePlayer->GetActiveCameraComponent());
 		
 		// When the sequence has finished, we'll tear down the test in this map via the OnSequenceFinished dispatch
 		// because TeardownTest's signature doesn't match OnFinished
@@ -196,6 +200,10 @@ void UAutomatedSequencePerfTest::OnCameraCut(UCameraComponent* CameraComponent)
 		if(RequestsCSVProfiler())
 		{
 			CSV_EVENT(AutomatedPerfTest, TEXT("END_%s"), *GetCameraCutID())
+			if(GetCSVOutputMode() == EAutomatedPerfTestCSVOutputMode::Granular)
+			{
+				TryStopCSVProfiler();
+			}
 		}
 #endif
 	}
@@ -203,6 +211,8 @@ void UAutomatedSequencePerfTest::OnCameraCut(UCameraComponent* CameraComponent)
 	// then null check the new CameraComponent so that we can use OnCameraCut to mark the end of the final camera cut region
 	if(CameraComponent)
 	{
+		NumCameraCuts += 1;
+		
 		// Then bring in the new camera component for this cut and mark the start of it
 		CurrentCamera = CameraComponent;
 		if(RequestsInsightsTrace())
@@ -212,7 +222,12 @@ void UAutomatedSequencePerfTest::OnCameraCut(UCameraComponent* CameraComponent)
 #if CSV_PROFILER
 		if(RequestsCSVProfiler())
 		{
+			if(GetCSVOutputMode() == EAutomatedPerfTestCSVOutputMode::Granular)
+			{
+				TryStartCSVProfiler(GetCameraCutFullName());
+			}
 			CSV_EVENT(AutomatedPerfTest, TEXT("START_%s"), *GetCameraCutID())
+			
 		}
 #endif
 	}
@@ -220,16 +235,37 @@ void UAutomatedSequencePerfTest::OnCameraCut(UCameraComponent* CameraComponent)
 
 FString UAutomatedSequencePerfTest::GetTestID()
 {
-	return CurrentMapSequenceCombo.IsSet() ? Super::GetTestID() + "_" + CurrentMapSequenceCombo->ComboName.ToString() : Super::GetTestID();
+	return CurrentMapSequenceCombo.IsSet() ? Super::GetTestID() + "_Sequence_" + CurrentMapSequenceCombo->ComboName.ToString() : Super::GetTestID();
 }
 
 FString UAutomatedSequencePerfTest::GetCameraCutID()
 {
 	if(CurrentCamera != nullptr)
 	{
-		const AActor* Owner = CurrentCamera->GetOwner();
-		FString CameraCutName = Owner ? Owner->GetActorNameOrLabel() : CurrentCamera->GetReadableName();
-		return GetTestID() + "_" + CameraCutName; 
+		// getting the label of a spawnable camera from Sequencer in a packaged build isn't possible
+		// via this method. TODO: find a more reliable way to return camera cut names that are set in Sequencerj
+		//const AActor* Owner = CurrentCamera->GetOwner();
+		//FString CutName = Owner ? Owner->GetActorNameOrLabel() : CurrentCamera->GetName();
+		//return CutName;
+		FString NumCuts = FString::FromInt(NumCameraCuts);
+
+		// Pad it out to at least four characters for now
+		while(NumCuts.Len() < 4)
+		{
+			NumCuts = "0" + NumCuts;
+		}
+		
+		return FString::Format(TEXT("CameraCut{0}"), {NumCuts});
+	}
+	return "";
+}
+
+FString UAutomatedSequencePerfTest::GetCameraCutFullName()
+{
+	// the full name for the camera cut should include the TestID so that the data can be 
+	if(CurrentCamera != nullptr)
+	{
+		return GetTestID() + "_" + GetCameraCutID(); 
 	}
 	return GetTestID();
 }
@@ -241,6 +277,8 @@ void UAutomatedSequencePerfTest::OnInit()
 	
 	Settings = GetDefault<UAutomatedSequencePerfTestProjectSettings>();
 
+	SetCSVOutputMode(Settings->CSVOutputMode);
+	
 	// if an explicit map/sequence name was set from commandline, use this to override the test
 	if (FParse::Value(FCommandLine::Get(), TEXT("AutomatedPerfTest.SequencePerfTest.MapSequenceName="), SequenceTestName))
 	{
diff --git a/Engine/Plugins/Performance/AutomatedPerfTesting/Source/AutomatedPerfTesting/Private/StaticCameraTests/AutomatedPerfTestStaticCamera.cpp b/Engine/Plugins/Performance/AutomatedPerfTesting/Source/AutomatedPerfTesting/Private/StaticCameraTests/AutomatedPerfTestStaticCamera.cpp
new file mode 100644
--- /dev/null
+++ b/Engine/Plugins/Performance/AutomatedPerfTesting/Source/AutomatedPerfTesting/Private/StaticCameraTests/AutomatedPerfTestStaticCamera.cpp
@@ -0,0 +1,10 @@
+﻿// Copyright Epic Games, Inc. All Rights Reserved.
+
+#include "StaticCameraTests/AutomatedPerfTestStaticCamera.h"
+
+AAutomatedPerfTestStaticCamera::AAutomatedPerfTestStaticCamera(const FObjectInitializer& ObjectInitializer)
+: Super(ObjectInitializer)
+{
+	// disable tick on these actors, as it's not needed
+	AActor::SetActorTickEnabled(false);
+}
diff --git a/Engine/Plugins/Performance/AutomatedPerfTesting/Source/AutomatedPerfTesting/Private/StaticCameraTests/AutomatedPlacedStaticCameraPerfTest.cpp b/Engine/Plugins/Performance/AutomatedPerfTesting/Source/AutomatedPerfTesting/Private/StaticCameraTests/AutomatedPlacedStaticCameraPerfTest.cpp
new file mode 100644
--- /dev/null
+++ b/Engine/Plugins/Performance/AutomatedPerfTesting/Source/AutomatedPerfTesting/Private/StaticCameraTests/AutomatedPlacedStaticCameraPerfTest.cpp
@@ -0,0 +1,22 @@
+﻿// Copyright Epic Games, Inc. All Rights Reserved.
+
+#include "StaticCameraTests/AutomatedPlacedStaticCameraPerfTest.h"
+#include "Kismet/GameplayStatics.h"
+#include "StaticCameraTests/AutomatedPerfTestStaticCamera.h"
+
+TArray<ACameraActor*> UAutomatedPlacedStaticCameraPerfTest::GetMapCameraActors()
+{
+	TArray<AActor*> FoundCameras;
+	TArray<ACameraActor*> OutCameras;
+	UGameplayStatics::GetAllActorsOfClass(GetWorld(), AAutomatedPerfTestStaticCamera::StaticClass(), FoundCameras);
+
+	for(AActor* Actor : FoundCameras)
+	{
+		if(ACameraActor* FoundCamera = Cast<ACameraActor>(Actor))
+		{
+			OutCameras.Add(FoundCamera);	
+		}
+	}
+
+	return OutCameras;
+}
diff --git a/Engine/Plugins/Performance/AutomatedPerfTesting/Source/AutomatedPerfTesting/Private/StaticCameraTests/AutomatedStaticCameraPerfTestBase.cpp b/Engine/Plugins/Performance/AutomatedPerfTesting/Source/AutomatedPerfTesting/Private/StaticCameraTests/AutomatedStaticCameraPerfTestBase.cpp
new file mode 100644
--- /dev/null
+++ b/Engine/Plugins/Performance/AutomatedPerfTesting/Source/AutomatedPerfTesting/Private/StaticCameraTests/AutomatedStaticCameraPerfTestBase.cpp
@@ -0,0 +1,300 @@
+﻿// Copyright Epic Games, Inc. All Rights Reserved.
+
+#include "StaticCameraTests/AutomatedStaticCameraPerfTestBase.h"
+#include "TimerManager.h"
+#include "Kismet/GameplayStatics.h"
+#include "AutomatedPerfTesting.h"
+#include "Camera/CameraActor.h"
+#include "Misc/CommandLine.h"
+
+UAutomatedStaticCameraPerfTestProjectSettings::UAutomatedStaticCameraPerfTestProjectSettings(const FObjectInitializer& Initializer)
+	: Super(Initializer)
+	, bCaptureScreenshots(true)
+	, WarmUpTime(5.0)
+	, SoakTime(5.0)
+	, CooldownTime(1.0)
+	, CSVOutputMode(EAutomatedPerfTestCSVOutputMode::Separate)
+{
+
+}
+
+bool UAutomatedStaticCameraPerfTestProjectSettings::GetMapFromAssetName(FString AssetName, FSoftObjectPath& OutSoftObjectPath) const
+{
+	for (FSoftObjectPath MapPath : MapsToTest)
+	{
+		if(MapPath.GetAssetName() == AssetName)
+		{
+			OutSoftObjectPath = MapPath;
+			return true;
+		}
+	}
+	return false;
+}
+
+void UAutomatedStaticCameraPerfTestBase::SetupTest()
+{
+	// load up into the map defined in project settings
+	if(&CurrentMapPath != nullptr)
+	{
+		if(GetCurrentMap() == CurrentMapPath.GetAssetName())
+		{
+			Super::SetupTest();
+			
+			UE_LOG(LogAutomatedPerfTest, Log, TEXT("UAutomatedStaticCameraPerfTestBase::SetupTest"))
+			// make sure the world exists, then create a sequence player
+			if(UWorld* const World = GetWorld())
+			{
+				CamerasToTest = GetMapCameraActors();
+
+				if(CamerasToTest.Num() <= 0)
+				{
+					UE_LOG(LogAutomatedPerfTest, Warning, TEXT("No cameras found in the map %s, skipping to next map"), *CurrentMapPath.GetAssetName());
+					NextMap();
+					return;
+				}
+
+				UE_LOG(LogAutomatedPerfTest, Log, TEXT("Found %hhd cameras to test in map %s"), CamerasToTest.Num(), *CurrentMapPath.GetAssetName());
+
+				GetFirstPlayerController()->SetCinematicMode(true, true, true);
+				
+				// delay for WarmUpDelay, and call RunTest
+				FTimerHandle UnusedHandle;
+				GetWorld()->GetTimerManager().SetTimer(UnusedHandle, this, &UAutomatedStaticCameraPerfTestBase::RunTest, 1.0, false, Settings->WarmUpTime);
+			}
+			// if we have an invalid world, we can't run the test, so we should bail out
+			else
+			{
+				UE_LOG(LogAutomatedPerfTest, Error, TEXT("Invalid World when starting UAutomatedStaticCameraPerfTest, exiting..."));
+				EndAutomatedPerfTest(1);
+			}
+		}
+		else
+		{
+			UE_LOG(LogAutomatedPerfTest, Log, TEXT("Current Map Name %s is not expected %s, calling NextMap."), *GetCurrentMap(), *CurrentMapPath.GetAssetName())
+			NextMap();
+		}
+	}
+	else
+	{
+		UE_LOG(LogAutomatedPerfTest, Log, TEXT("Current Map Path has not been set, calling NextMap to trigger."))
+		NextMap();
+	}
+}
+
+void UAutomatedStaticCameraPerfTestBase::RunTest()
+{
+	Super::RunTest();
+
+	UE_LOG(LogAutomatedPerfTest, Log, TEXT("UAutomatedStaticCameraPerfTestBase::RunTest"))
+	
+	FTimerHandle UnusedHandle;
+	GetWorld()->GetTimerManager().SetTimer(UnusedHandle, this, &UAutomatedStaticCameraPerfTestBase::SetUpNextCamera, 1.0, false, Settings->WarmUpTime);
+}
+
+FString UAutomatedStaticCameraPerfTestBase::GetTestID()
+{
+	return Super::GetTestID() + "_StaticCamera_" + GetCurrentMap();
+}
+
+void UAutomatedStaticCameraPerfTestBase::SetUpNextCamera()
+{
+	if(CamerasToTest.Num() <= 0)
+	{
+		UE_LOG(LogAutomatedPerfTest, Log, TEXT("No more cameras left to test, moving to next map."))
+		NextMap();
+		return;
+	}
+
+	CurrentCamera = CamerasToTest.Pop();
+
+	UE_LOG(LogAutomatedPerfTest, Log, TEXT("Setting up %s to test"), *CurrentCamera->GetActorNameOrLabel());
+
+	GetFirstPlayerController()->SetViewTarget(CurrentCamera);
+	FVector ViewLocation = CurrentCamera->GetActorLocation();
+	FRotator ViewRotation = CurrentCamera->GetActorRotation();
+
+	FString GoString = FString::Printf(TEXT("BugItGo %f %f %f %f %f %f"), ViewLocation.X, ViewLocation.Y, ViewLocation.Z, ViewRotation.Pitch, ViewRotation.Yaw, ViewRotation.Roll);
+	
+	FTimerHandle UnusedHandle;
+	GetWorld()->GetTimerManager().SetTimer(UnusedHandle, this, &UAutomatedStaticCameraPerfTestBase::EvaluateCamera, 1.0, false, Settings->WarmUpTime);
+}
+
+void UAutomatedStaticCameraPerfTestBase::EvaluateCamera()
+{
+	MarkCameraStart();
+
+	FTimerHandle UnusedHandle;
+	GetWorld()->GetTimerManager().SetTimer(UnusedHandle, this, &UAutomatedStaticCameraPerfTestBase::FinishCamera, 1.0, false, Settings->SoakTime);
+}
+
+void UAutomatedStaticCameraPerfTestBase::FinishCamera()
+{
+	MarkCameraEnd();
+	
+	if(Settings->bCaptureScreenshots)
+	{
+		FTimerHandle UnusedHandle;
+		GetWorld()->GetTimerManager().SetTimer(UnusedHandle, this, &UAutomatedStaticCameraPerfTestBase::ScreenshotCamera, 1.0, false, Settings->CooldownTime);		
+	}
+	else
+	{
+		FTimerHandle UnusedHandle;
+		GetWorld()->GetTimerManager().SetTimer(UnusedHandle, this, &UAutomatedStaticCameraPerfTestBase::SetUpNextCamera, 1.0, false, Settings->CooldownTime);
+	}
+}
+
+void UAutomatedStaticCameraPerfTestBase::ScreenshotCamera()
+{
+	TakeScreenshot(GetCurrentCameraRegionName());
+
+	FTimerHandle UnusedHandle;
+	GetWorld()->GetTimerManager().SetTimer(UnusedHandle, this, &UAutomatedStaticCameraPerfTestBase::SetUpNextCamera, 1.0, false, Settings->CooldownTime);
+}
+
+void UAutomatedStaticCameraPerfTestBase::NextMap()
+{
+	UE_LOG(LogAutomatedPerfTest, Log, TEXT("UAutomatedStaticCameraPerfTestBase::NextMap"))
+
+	if(MapsToTest.Num() > 0)
+	{
+		CurrentMapPath = MapsToTest.Pop();
+		UE_LOG(LogAutomatedPerfTest, Log, TEXT("Setting up test for Map %s"), *CurrentMapPath.GetAssetName())
+
+		// no need to prepend this with a ? since OpenLevel handles that part for us
+		FString OptionsString;
+		if(!Settings->GameModeOverride.IsEmpty())
+		{
+			UE_LOG(LogAutomatedPerfTest, Log, TEXT("Game Mode overridden to %s"), *Settings->GameModeOverride)
+			OptionsString += "game=" + Settings->GameModeOverride;
+		}
+		
+		UE_LOG(LogAutomatedPerfTest, Log, TEXT("Opening map %s%s"), *CurrentMapPath.GetAssetName(), *OptionsString);
+		UGameplayStatics::OpenLevel(AutomatedPerfTest::FindCurrentWorld(), *CurrentMapPath.GetAssetName(), true, OptionsString);
+	}
+	else
+	{
+		UE_LOG(LogAutomatedPerfTest, Log, TEXT("UAutomatedStaticCameraPerfTestBase::NextMap, all maps complete, exiting after delay."))
+		TriggerExitAfterDelay();
+	}
+}
+
+TArray<ACameraActor*> UAutomatedStaticCameraPerfTestBase::GetMapCameraActors()
+{
+	UE_LOG(LogAutomatedPerfTest, Warning, TEXT("GetMapCameraActors called in base class UAutomatedStaticCameraPerfTestBase, please ensure you've overridden this function in the subclass, and you're not using the base class as your test controller"));
+	return TArray<ACameraActor*>();
+}
+
+ACameraActor* UAutomatedStaticCameraPerfTestBase::GetCurrentCamera() const
+{
+	return CurrentCamera;
+}
+
+FString UAutomatedStaticCameraPerfTestBase::GetCurrentCameraRegionName()
+{
+	return GetCurrentCamera()->GetActorNameOrLabel();
+}
+
+FString UAutomatedStaticCameraPerfTestBase::GetCurrentCameraRegionFullName()
+{
+	return GetTestID() + "_" + GetCurrentCameraRegionName();
+}
+
+void UAutomatedStaticCameraPerfTestBase::MarkCameraStart()
+{
+	// safety check on the current material
+	if(CurrentCamera)
+	{
+		if(RequestsInsightsTrace())
+		{
+			TRACE_BEGIN_REGION(*GetCurrentCameraRegionName());
+		}
+#if CSV_PROFILER
+		if(RequestsCSVProfiler())
+		{
+			if(GetCSVOutputMode() == EAutomatedPerfTestCSVOutputMode::Granular)
+			{
+				TryStartCSVProfiler(GetCurrentCameraRegionFullName());
+			}
+			CSV_EVENT(AutomatedPerfTest, TEXT("START_%s"), *GetCurrentCameraRegionName())
+		}
+#endif
+	}
+}
+
+void UAutomatedStaticCameraPerfTestBase::MarkCameraEnd()
+{
+	// safety check on the current material
+	if(CurrentCamera)
+	{
+		if(RequestsInsightsTrace())
+		{
+			TRACE_END_REGION(*GetCurrentCameraRegionName());
+		}
+#if CSV_PROFILER
+		if(RequestsCSVProfiler())
+		{
+			CSV_EVENT(AutomatedPerfTest, TEXT("END_%s"), *GetCurrentCameraRegionName())
+			if(GetCSVOutputMode() == EAutomatedPerfTestCSVOutputMode::Granular)
+			{
+				TryStopCSVProfiler();
+			}
+		}
+#endif
+	}
+}
+
+void UAutomatedStaticCameraPerfTestBase::OnInit()
+{
+	Super::OnInit();
+	
+	UE_LOG(LogAutomatedPerfTest, Log, TEXT("UAutomatedStaticCameraPerfTestBase::OnInit"));
+
+	Settings = GetDefault<UAutomatedStaticCameraPerfTestProjectSettings>();
+
+	SetCSVOutputMode(Settings->CSVOutputMode);
+	
+	// if an explicit map/sequence name was set from commandline, use this to override the test
+	if (FParse::Value(FCommandLine::Get(), TEXT("AutomatedPerfTest.StaticCameraPerfTest.MapName="), CurrentMapName))
+	{
+		for (FSoftObjectPath MapPath : Settings->MapsToTest)
+		{
+			if(MapPath.GetAssetName() == CurrentMapName)
+			{
+				MapsToTest.Add(MapPath);
+			}
+		}
+		if(MapsToTest.Num() == 0)
+		{
+			UE_LOG(LogAutomatedPerfTest, Error, TEXT("Couldn't find a map name matching %s in Static Camera Maps to Test setting. Exiting."), *CurrentMapName)
+			EndAutomatedPerfTest(1);
+		}
+	}
+	// otherwise, use all the maps defined in project settings
+	else
+	{
+		UE_LOG(LogAutomatedPerfTest, Log, TEXT("No map name specified, testing all maps."))
+		for(FSoftObjectPath MapName : Settings->MapsToTest)
+		{
+			MapsToTest.Add(MapName);
+		}
+	}
+	
+	UE_LOG(LogAutomatedPerfTest, Log, TEXT("Initialized with %hhd MapSequence combos"), MapsToTest.Num());
+	
+	// early out if there aren't actually any maps set in project settings
+	if(Settings->MapsToTest.Num() < 0)
+	{
+		UE_LOG(LogAutomatedPerfTest, Error, TEXT("No maps defined in the project's Automated Perf Test | Static Camera settings. Exiting test early."));
+		EndAutomatedPerfTest(1);
+	}
+}
+
+void UAutomatedStaticCameraPerfTestBase::UnbindAllDelegates()
+{
+	Super::UnbindAllDelegates();
+	
+	GetWorld()->GetTimerManager().ClearAllTimersForObject(this);
+
+	CurrentCamera = nullptr;
+	CamerasToTest.Empty();
+}
\ No newline at end of file
diff --git a/Engine/Plugins/Performance/AutomatedPerfTesting/Source/AutomatedPerfTesting/Public/AutomatedMaterialPerfTest.h b/Engine/Plugins/Performance/AutomatedPerfTesting/Source/AutomatedPerfTesting/Public/AutomatedMaterialPerfTest.h
new file mode 100644
--- /dev/null
+++ b/Engine/Plugins/Performance/AutomatedPerfTesting/Source/AutomatedPerfTesting/Public/AutomatedMaterialPerfTest.h
@@ -0,0 +1,150 @@
+﻿// Copyright Epic Games, Inc. All Rights Reserved.
+#pragma once
+
+#include "AutomatedPerfTestControllerBase.h"
+#include "AutomatedPerfTestProjectSettings.h"
+#include "UObject/Object.h"
+#include "Templates/SubclassOf.h"
+#include "Engine/EngineTypes.h"
+#include "Engine/DeveloperSettings.h"
+#include "Camera/CameraTypes.h"
+#include "AutomatedMaterialPerfTest.generated.h"
+
+class AStaticMeshActor;
+
+
+UCLASS(BlueprintType, Config=Engine, DefaultConfig, DisplayName="Automated Performance Testing | Materials")
+class AUTOMATEDPERFTESTING_API UAutomatedMaterialPerfTestProjectSettings : public UDeveloperSettings
+{
+	GENERATED_BODY()
+
+public:
+	UAutomatedMaterialPerfTestProjectSettings(const FObjectInitializer&);
+
+	/** Gets the settings container name for the settings, either Project or Editor */
+	virtual FName GetContainerName() const override { return FName("Project"); }
+	/** Gets the category for the settings, some high level grouping like, Editor, Engine, Game...etc. */
+	virtual FName GetCategoryName() const override { return FName("Plugins"); }
+	
+	/*
+	 * List of materials to load and review for the Material Performance Test
+	 */
+	UPROPERTY(Config, EditAnywhere, BlueprintReadWrite, Category="Material Perf Test | Parameters", meta=(AllowedClasses="/Script/Engine.MaterialInterface"))
+	TArray<FSoftObjectPath> MaterialsToTest;
+
+	/*
+	 * If true, will capture a screenshot for each material tested after gathering data
+	 */
+	UPROPERTY(Config, EditAnywhere, BlueprintReadWrite, Category="Material Perf Test | Parameters")
+	bool bCaptureScreenshots;
+
+	/*
+	 * For how long the material performance test should delay before beginning to gather data for a material, in seconds
+     */
+	UPROPERTY(Config, EditAnywhere, BlueprintReadWrite, Category="Material Perf Test | Parameters")
+	float WarmUpTime;
+	
+	/*
+	 * For how long the material performance test should gather data on each material, in seconds
+	 */
+	UPROPERTY(Config, EditAnywhere, BlueprintReadWrite, Category="Material Perf Test | Parameters")
+	float SoakTime;
+
+	/*
+	 * For how long the material performance test should delay after ending evaluation before switching to the next material
+     */
+	UPROPERTY(Config, EditAnywhere, BlueprintReadWrite, Category="Material Perf Test | Parameters")
+	float CooldownTime;
+
+	/*
+	 * The map in which the material test will take place. Useful if you need to set up things like RVT volumes that are required
+	 * by your materials.
+	 */
+	UPROPERTY(Config, EditAnywhere, BlueprintReadWrite, Category="Material Perf Test | Scene", meta=(AllowedClasses="/Script/Engine.World"))
+	FSoftObjectPath MaterialPerformanceTestMap;
+
+	/*
+	 * If set, will launch the material performance test map with this game mode alias (make sure you've set the game mode alias in
+	 * the Maps and Modes settings of your project!)
+	 */
+	UPROPERTY(Config, EditAnywhere, BlueprintReadWrite, Category="Material Perf Test | Scene")
+	FString GameModeOverride;
+	
+	/*
+     * Which camera projection mode to use
+     */
+	UPROPERTY(Config, EditAnywhere, BlueprintReadWrite, Category="Material Perf Test | Camera")
+	TEnumAsByte<ECameraProjectionMode::Type> CameraProjectionMode;
+	
+	/*
+     * How far away from the camera should the material test plate be placed
+     * (will also be used for Ortho Width if the camera projection mode is orthographic)
+     */
+	UPROPERTY(Config, EditAnywhere, BlueprintReadWrite, Category="Material Perf Test | Plate")
+	float PlateDistanceFromCamera;
+
+	/*
+     * Which static mesh to use as the material plate
+     */
+	UPROPERTY(Config, EditAnywhere, BlueprintReadWrite, Category="Material Perf Test | Plate")
+	FSoftObjectPath MaterialPlate;
+	
+	/*
+	 * For Material Perf Tests, Granular will output one CSV per-material. 
+	 */
+	UPROPERTY(Config, EditAnywhere, BlueprintReadWrite, Category="Material Perf Test | CSV")
+	EAutomatedPerfTestCSVOutputMode CSVOutputMode;
+};
+
+/**
+ * 
+ */
+UCLASS()
+class AUTOMATEDPERFTESTING_API UAutomatedMaterialPerfTest : public UAutomatedPerfTestControllerBase
+{
+	GENERATED_BODY()
+	
+public:
+	// ~Begin UAutomatedPerfTestControllerBase Interface
+	virtual void SetupTest() override;
+	
+	UFUNCTION()
+	virtual void RunTest() override;
+
+	virtual FString GetTestID() override;
+	// ~End UAutomatedPerfTestControllerBase Interface
+	
+	UFUNCTION()
+	void SetUpNextMaterial();
+
+	UFUNCTION()
+	void EvaluateMaterial();
+	
+	UFUNCTION()
+	void FinishMaterialEvaluation();
+	
+	UFUNCTION()
+	void ScreenshotMaterial();
+
+	UMaterialInterface* GetCurrentMaterial() const;
+	FString GetCurrentMaterialRegionName();
+	FString GetCurrentMaterialRegionFullName();
+
+	void MarkMaterialStart();
+	void MarkMaterialEnd();
+	
+protected:
+	virtual void OnInit() override;
+	virtual void UnbindAllDelegates() override;
+
+	void OpenMaterialPerformanceTestMap() const;
+
+private:
+	ACameraActor* Camera;
+	AStaticMeshActor* MaterialPlate;
+	FString CurrentRegionName;
+	UMaterialInterface* CurrentMaterial;
+	int CurrentMaterialIndex;
+	const UAutomatedMaterialPerfTestProjectSettings* Settings;
+};
+
diff --git a/Engine/Plugins/Performance/AutomatedPerfTesting/Source/AutomatedPerfTesting/Public/AutomatedPerfTestControllerBase.h b/Engine/Plugins/Performance/AutomatedPerfTesting/Source/AutomatedPerfTesting/Public/AutomatedPerfTestControllerBase.h
--- a/Engine/Plugins/Performance/AutomatedPerfTesting/Source/AutomatedPerfTesting/Public/AutomatedPerfTestControllerBase.h
+++ b/Engine/Plugins/Performance/AutomatedPerfTesting/Source/AutomatedPerfTesting/Public/AutomatedPerfTestControllerBase.h
@@ -3,14 +3,23 @@
 #pragma once
 
 #include "CoreMinimal.h"
+#include "AutomatedPerfTestProjectSettings.h"
 #include "GauntletTestController.h"
 #include "Engine/World.h"
 #include "Engine/Engine.h"
 
 #include "AutomatedPerfTestControllerBase.generated.h"
 
-class AGameModeBase;
+UENUM()
+enum class EAutomatedPerfTestCSVOutputMode : uint8
+{
+	Single UMETA(DisplayName = "Single CSV", ToolTip = "Output a single CSV with all of the results for the entire session, from SetupTest to ExitTest."),
+	Separate UMETA(DisplayName = "Separate CSVs", ToolTip = "Output CSVs from RunTest to TeardownTest. May result into multiple output CSVs that require special processing."),
+	Granular UMETA(DisplayName = "Granular CSVs", ToolTip = "Output granular CSVs during the test run, resulting in multiple CSVs between RunTest and TeardownTest.")
+};
+
 
+class AGameModeBase;
 
 namespace AutomatedPerfTest
 {
@@ -77,6 +86,7 @@ public:
 	bool TryStopInsightsTrace();
 
 	bool TryStartCSVProfiler();
+	bool TryStartCSVProfiler(FString CSVFileName);
 	bool TryStopCSVProfiler();
 	
 	bool TryStartFPSChart();
@@ -92,6 +102,12 @@ public:
 	virtual void Exit();
 
 	AGameModeBase* GetGameMode() const;
+
+	void TakeScreenshot(FString ScreenshotName);
+
+	// you'll need to set this via your subclass if you want to customize the behavior otherwise it will default to a single CSV per session
+	void SetCSVOutputMode(EAutomatedPerfTestCSVOutputMode NewOutputMode);
+	EAutomatedPerfTestCSVOutputMode GetCSVOutputMode() const;
 	
 protected:
 	// ~Begin UGauntletTestController Interface
@@ -127,4 +143,6 @@ private:
 	AGameModeBase* GameMode;
 
 	FDelegateHandle CsvProfilerDelegateHandle;
+
+	EAutomatedPerfTestCSVOutputMode CSVOutputMode;
 };
diff --git a/Engine/Plugins/Performance/AutomatedPerfTesting/Source/AutomatedPerfTesting/Public/AutomatedPerfTestProjectSettings.h b/Engine/Plugins/Performance/AutomatedPerfTesting/Source/AutomatedPerfTesting/Public/AutomatedPerfTestProjectSettings.h
--- a/Engine/Plugins/Performance/AutomatedPerfTesting/Source/AutomatedPerfTesting/Public/AutomatedPerfTestProjectSettings.h
+++ b/Engine/Plugins/Performance/AutomatedPerfTesting/Source/AutomatedPerfTesting/Public/AutomatedPerfTestProjectSettings.h
@@ -8,7 +8,7 @@
 #include "AutomatedPerfTestProjectSettings.generated.h"
 
 
-UCLASS(BlueprintType, Config=Engine, DefaultConfig, DisplayName="Automated Performance Testing")
+UCLASS(Config=Engine)
 class AUTOMATEDPERFTESTING_API UAutomatedPerfTestProjectSettings : public UDeveloperSettings
 {
 	GENERATED_BODY()
diff --git a/Engine/Plugins/Performance/AutomatedPerfTesting/Source/AutomatedPerfTesting/Public/AutomatedSequencePerfTest.h b/Engine/Plugins/Performance/AutomatedPerfTesting/Source/AutomatedPerfTesting/Public/AutomatedSequencePerfTest.h
--- a/Engine/Plugins/Performance/AutomatedPerfTesting/Source/AutomatedPerfTesting/Public/AutomatedSequencePerfTest.h
+++ b/Engine/Plugins/Performance/AutomatedPerfTesting/Source/AutomatedPerfTesting/Public/AutomatedSequencePerfTest.h
@@ -70,6 +70,12 @@ public:
 	/* How long to delay between setting up the sequence for the map before the sequence actually starts */
 	UPROPERTY(Config, EditAnywhere, BlueprintReadWrite, Category="Sequence Perf Test")
 	float SequenceStartDelay=5.0;
+
+	/*
+	 * For Sequence Perf Tests, Separate will output one CSV per map tested, and Granular will output one CSV per camera-cut. 
+	 */
+	UPROPERTY(Config, EditAnywhere, BlueprintReadWrite, Category="Sequence Perf Test")
+	EAutomatedPerfTestCSVOutputMode CSVOutputMode;
 };
 
 /**
@@ -103,6 +109,7 @@ public:
 	void OnCameraCut(UCameraComponent* CameraComponent);
 	
 	FString GetCameraCutID();
+	FString GetCameraCutFullName();
 	
 protected:
 	virtual void OnInit() override;
@@ -119,4 +126,5 @@ private:
 	ULevelSequencePlayer* SequencePlayer;
 
 	UCameraComponent* CurrentCamera;
+	int NumCameraCuts; // TODO this is temporary until a reliable method is found for getting the user-set name of a camera cut out of Sequencer in packaged builds
 };
diff --git a/Engine/Plugins/Performance/AutomatedPerfTesting/Source/AutomatedPerfTesting/Public/StaticCameraTests/AutomatedPerfTestStaticCamera.h b/Engine/Plugins/Performance/AutomatedPerfTesting/Source/AutomatedPerfTesting/Public/StaticCameraTests/AutomatedPerfTestStaticCamera.h
new file mode 100644
--- /dev/null
+++ b/Engine/Plugins/Performance/AutomatedPerfTesting/Source/AutomatedPerfTesting/Public/StaticCameraTests/AutomatedPerfTestStaticCamera.h
@@ -0,0 +1,13 @@
+﻿// Copyright Epic Games, Inc. All Rights Reserved.
+#pragma once
+
+#include "Camera/CameraActor.h"
+#include "AutomatedPerfTestStaticCamera.generated.h"
+
+UCLASS(Blueprintable, BlueprintType, ClassGroup=Performance)
+class AUTOMATEDPERFTESTING_API AAutomatedPerfTestStaticCamera : public ACameraActor
+{
+	GENERATED_BODY()
+
+	AAutomatedPerfTestStaticCamera(const FObjectInitializer& ObjectInitializer);
+};
\ No newline at end of file
diff --git a/Engine/Plugins/Performance/AutomatedPerfTesting/Source/AutomatedPerfTesting/Public/StaticCameraTests/AutomatedPlacedStaticCameraPerfTest.h b/Engine/Plugins/Performance/AutomatedPerfTesting/Source/AutomatedPerfTesting/Public/StaticCameraTests/AutomatedPlacedStaticCameraPerfTest.h
new file mode 100644
--- /dev/null
+++ b/Engine/Plugins/Performance/AutomatedPerfTesting/Source/AutomatedPerfTesting/Public/StaticCameraTests/AutomatedPlacedStaticCameraPerfTest.h
@@ -0,0 +1,17 @@
+﻿// Copyright Epic Games, Inc. All Rights Reserved.
+#pragma once
+
+#include "StaticCameraTests/AutomatedStaticCameraPerfTestBase.h"
+#include "UObject/Object.h"
+#include "Templates/SubclassOf.h"
+#include "Engine/EngineTypes.h"
+#include "AutomatedPlacedStaticCameraPerfTest.generated.h"
+
+UCLASS()
+class AUTOMATEDPERFTESTING_API UAutomatedPlacedStaticCameraPerfTest : public UAutomatedStaticCameraPerfTestBase
+{
+	GENERATED_BODY()
+
+public:
+	virtual TArray<ACameraActor*> GetMapCameraActors() override;
+};
\ No newline at end of file
diff --git a/Engine/Plugins/Performance/AutomatedPerfTesting/Source/AutomatedPerfTesting/Public/StaticCameraTests/AutomatedStaticCameraPerfTestBase.h b/Engine/Plugins/Performance/AutomatedPerfTesting/Source/AutomatedPerfTesting/Public/StaticCameraTests/AutomatedStaticCameraPerfTestBase.h
new file mode 100644
--- /dev/null
+++ b/Engine/Plugins/Performance/AutomatedPerfTesting/Source/AutomatedPerfTesting/Public/StaticCameraTests/AutomatedStaticCameraPerfTestBase.h
@@ -0,0 +1,129 @@
+﻿// Copyright Epic Games, Inc. All Rights Reserved.
+#pragma once
+
+#include "AutomatedPerfTestControllerBase.h"
+#include "UObject/Object.h"
+#include "Engine/DeveloperSettings.h"
+#include "AutomatedStaticCameraPerfTestBase.generated.h"
+
+
+UCLASS(BlueprintType, Config=Engine, DefaultConfig, DisplayName="Automated Performance Testing | Static Camera")
+class AUTOMATEDPERFTESTING_API UAutomatedStaticCameraPerfTestProjectSettings : public UDeveloperSettings
+{
+	GENERATED_BODY()
+
+	TMap<FString, FSoftObjectPath> MapNameMap; 
+	
+public:
+	UAutomatedStaticCameraPerfTestProjectSettings(const FObjectInitializer&);
+
+	/** Gets the settings container name for the settings, either Project or Editor */
+	virtual FName GetContainerName() const override { return FName("Project"); }
+	/** Gets the category for the settings, some high level grouping like, Editor, Engine, Game...etc. */
+	virtual FName GetCategoryName() const override { return FName("Plugins"); }
+
+	UFUNCTION(BlueprintCallable, Category="Static Camera Perf Test")
+	bool GetMapFromAssetName(FString AssetName, FSoftObjectPath& OutSoftObjectPath) const;
+	
+	/*
+	 * List of levels to test
+	 */
+	UPROPERTY(Config, EditAnywhere, BlueprintReadWrite, Category="Static Cameras", meta=(AllowedClasses="/Script/Engine.World"))
+	TArray<FSoftObjectPath> MapsToTest;
+
+	/*
+     * If set, will launch the material performance test map with this game mode alias (make sure you've set the game mode alias in
+     * the Maps and Modes settings of your project!)
+     */
+	UPROPERTY(Config, EditAnywhere, BlueprintReadWrite, Category="Static Cameras")
+	FString GameModeOverride;
+	
+	/*
+	 * If true, will capture a screenshot for each camera tested after gathering data
+	 */
+	UPROPERTY(Config, EditAnywhere, BlueprintReadWrite, Category="Static Cameras")
+	bool bCaptureScreenshots;
+
+	/*
+	 * For how long the material performance test should delay before beginning to gather data for a material, in seconds
+     */
+	UPROPERTY(Config, EditAnywhere, BlueprintReadWrite, Category="Static Cameras")
+	float WarmUpTime;
+	
+	/*
+	 * For how long the static camera performance test should gather data on each camera, in seconds
+	 */
+	UPROPERTY(Config, EditAnywhere, BlueprintReadWrite, Category="Static Cameras")
+	float SoakTime;
+
+	/*
+	 * For how long the static camera performance test should delay after ending evaluation before switching to the next camera
+     */
+	UPROPERTY(Config, EditAnywhere, BlueprintReadWrite, Category="Static Cameras")
+	float CooldownTime;
+
+	/*
+	 * For Static Camera Perf Tests, Separate will output one CSV per map tested, and Granular will output one CSV per camera. 
+	 */
+	UPROPERTY(Config, EditAnywhere, BlueprintReadWrite, Category="Static Cameras")
+	EAutomatedPerfTestCSVOutputMode CSVOutputMode;
+};
+
+/**
+ * 
+ */
+UCLASS()
+class AUTOMATEDPERFTESTING_API UAutomatedStaticCameraPerfTestBase : public UAutomatedPerfTestControllerBase
+{
+	GENERATED_BODY()
+	
+public:
+	// ~Begin UAutomatedPerfTestControllerBase Interface
+	virtual void SetupTest() override;
+	
+	UFUNCTION()
+	virtual void RunTest() override;
+
+	virtual FString GetTestID() override;
+	// ~End UAutomatedPerfTestControllerBase Interface
+	
+	UFUNCTION()
+	void SetUpNextCamera();
+
+	UFUNCTION()
+	void EvaluateCamera();
+	
+	UFUNCTION()
+	void FinishCamera();
+	
+	UFUNCTION()
+	void ScreenshotCamera();
+
+	UFUNCTION()
+	void NextMap();
+
+	virtual TArray<ACameraActor*> GetMapCameraActors();
+	
+	ACameraActor* GetCurrentCamera() const;
+	FString GetCurrentCameraRegionName();
+	FString GetCurrentCameraRegionFullName();
+
+	void MarkCameraStart();
+	void MarkCameraEnd();
+	
+protected:
+	virtual void OnInit() override;
+	virtual void UnbindAllDelegates() override;
+
+private:
+	UPROPERTY()
+	TArray<TObjectPtr<ACameraActor>> CamerasToTest;
+
+	UPROPERTY()
+	TObjectPtr<ACameraActor> CurrentCamera;
+	
+	FString CurrentMapName;
+	FSoftObjectPath CurrentMapPath;
+	TArray<FSoftObjectPath> MapsToTest;
+	const UAutomatedStaticCameraPerfTestProjectSettings* Settings;
+};
