AutomatedPerfTesting - Disabling trace screenshots in shipping by default

diff --git a/Engine/Plugins/Performance/AutomatedPerfTesting/Source/AutomatedPerfTesting/Private/AutomatedPerfTestControllerBase.cpp b/Engine/Plugins/Performance/AutomatedPerfTesting/Source/AutomatedPerfTesting/Private/AutomatedPerfTestControllerBase.cpp
--- a/Engine/Plugins/Performance/AutomatedPerfTesting/Source/AutomatedPerfTesting/Private/AutomatedPerfTestControllerBase.cpp
+++ b/Engine/Plugins/Performance/AutomatedPerfTesting/Source/AutomatedPerfTesting/Private/AutomatedPerfTestControllerBase.cpp
@@ -459,7 +459,10 @@ void UAutomatedPerfTestControllerBase::TakeScreenshot(FString ScreenshotName)
 {
 	if(RequestsInsightsTrace())
 	{
+		// trace screenshots are disabled in shipping by default
+#if UE_SCREENSHOT_TRACE_ENABLED
 		FTraceScreenshot::RequestScreenshot(ScreenshotName, false, LogAutomatedPerfTest);
+#endif
 	}
 	else
 	{
