Fix warning C6398: The address-of a field cannot be null in well-defined code.
#rnx

diff --git a/Engine/Plugins/Performance/AutomatedPerfTesting/Source/AutomatedPerfTesting/Private/StaticCameraTests/AutomatedStaticCameraPerfTestBase.cpp b/Engine/Plugins/Performance/AutomatedPerfTesting/Source/AutomatedPerfTesting/Private/StaticCameraTests/AutomatedStaticCameraPerfTestBase.cpp
--- a/Engine/Plugins/Performance/AutomatedPerfTesting/Source/AutomatedPerfTesting/Private/StaticCameraTests/AutomatedStaticCameraPerfTestBase.cpp
+++ b/Engine/Plugins/Performance/AutomatedPerfTesting/Source/AutomatedPerfTesting/Private/StaticCameraTests/AutomatedStaticCameraPerfTestBase.cpp
@@ -34,7 +34,7 @@ bool UAutomatedStaticCameraPerfTestProjectSettings::GetMapFromAssetName(FString
 void UAutomatedStaticCameraPerfTestBase::SetupTest()
 {
 	// load up into the map defined in project settings
-	if(&CurrentMapPath != nullptr)
+	if(!CurrentMapPath.IsNull())
 	{
 		if(GetCurrentMap() == CurrentMapPath.GetAssetName())
 		{
