AutomatedPerfTesting - Fixing missing includes

diff --git a/Engine/Plugins/Performance/AutomatedPerfTesting/Source/AutomatedPerfTesting/Private/AutomatedMaterialPerfTest.cpp b/Engine/Plugins/Performance/AutomatedPerfTesting/Source/AutomatedPerfTesting/Private/AutomatedMaterialPerfTest.cpp
--- a/Engine/Plugins/Performance/AutomatedPerfTesting/Source/AutomatedPerfTesting/Private/AutomatedMaterialPerfTest.cpp
+++ b/Engine/Plugins/Performance/AutomatedPerfTesting/Source/AutomatedPerfTesting/Private/AutomatedMaterialPerfTest.cpp
@@ -11,6 +11,7 @@
 #include "Engine/Player.h"
 #include "GameFramework/Pawn.h"
 #include "Materials/MaterialInterface.h"
+#include "ProfilingDebugging/CsvProfiler.h"
 
 UAutomatedMaterialPerfTestProjectSettings::UAutomatedMaterialPerfTestProjectSettings(const FObjectInitializer& Initializer)
 	: Super(Initializer)
@@ -101,6 +102,7 @@ FString UAutomatedMaterialPerfTest::GetTestID()
 
 bool UAutomatedMaterialPerfTest::TryStartCSVProfiler(FString CSVFileName)
 {
+#if CSV_PROFILER
 	if(FCsvProfiler* const CsvProfiler = FCsvProfiler::Get())
 	{
 		if (GetCSVOutputMode() == EAutomatedPerfTestCSVOutputMode::Granular)
@@ -109,7 +111,7 @@ bool UAutomatedMaterialPerfTest::TryStartCSVProfiler(FString CSVFileName)
 		}
 		return Super::TryStartCSVProfiler(CSVFileName);
 	}
-	
+#endif
 	UE_LOG(LogAutomatedPerfTest, Warning, TEXT("CSVProfiler Start requested, but not available."))
 	return false;
 }
diff --git a/Engine/Plugins/Performance/AutomatedPerfTesting/Source/AutomatedPerfTesting/Private/AutomatedSequencePerfTest.cpp b/Engine/Plugins/Performance/AutomatedPerfTesting/Source/AutomatedPerfTesting/Private/AutomatedSequencePerfTest.cpp
--- a/Engine/Plugins/Performance/AutomatedPerfTesting/Source/AutomatedPerfTesting/Private/AutomatedSequencePerfTest.cpp
+++ b/Engine/Plugins/Performance/AutomatedPerfTesting/Source/AutomatedPerfTesting/Private/AutomatedSequencePerfTest.cpp
@@ -10,6 +10,7 @@
 #include "MovieSceneSequencePlayer.h"
 #include "Kismet/GameplayStatics.h"
 #include "Misc/CommandLine.h"
+#include "ProfilingDebugging/CsvProfiler.h"
 
 /*****
  * UAutomatedSequencePerfTestProjectSettings
@@ -189,6 +190,7 @@ void UAutomatedSequencePerfTest::Exit()
 
 bool UAutomatedSequencePerfTest::TryStartCSVProfiler(FString CSVFileName)
 {
+#if CSV_PROFILER
 	if(FCsvProfiler* const CsvProfiler = FCsvProfiler::Get())
 	{
 		if (GetCSVOutputMode() != EAutomatedPerfTestCSVOutputMode::Single)
@@ -201,7 +203,7 @@ bool UAutomatedSequencePerfTest::TryStartCSVProfiler(FString CSVFileName)
 		}
 		return Super::TryStartCSVProfiler(CSVFileName);
 	}
-	
+#endif
 	UE_LOG(LogAutomatedPerfTest, Warning, TEXT("CSVProfiler Start requested, but not available."))
 	return false;
 }
diff --git a/Engine/Plugins/Performance/AutomatedPerfTesting/Source/AutomatedPerfTesting/Private/StaticCameraTests/AutomatedStaticCameraPerfTestBase.cpp b/Engine/Plugins/Performance/AutomatedPerfTesting/Source/AutomatedPerfTesting/Private/StaticCameraTests/AutomatedStaticCameraPerfTestBase.cpp
--- a/Engine/Plugins/Performance/AutomatedPerfTesting/Source/AutomatedPerfTesting/Private/StaticCameraTests/AutomatedStaticCameraPerfTestBase.cpp
+++ b/Engine/Plugins/Performance/AutomatedPerfTesting/Source/AutomatedPerfTesting/Private/StaticCameraTests/AutomatedStaticCameraPerfTestBase.cpp
@@ -1,4 +1,4 @@
-﻿// Copyright Epic Games, Inc. All Rights Reserved.
+// Copyright Epic Games, Inc. All Rights Reserved.
 
 #include "StaticCameraTests/AutomatedStaticCameraPerfTestBase.h"
 #include "TimerManager.h"
@@ -6,6 +6,7 @@
 #include "AutomatedPerfTesting.h"
 #include "Camera/CameraActor.h"
 #include "Misc/CommandLine.h"
+#include "ProfilingDebugging/CsvProfiler.h"
 
 UAutomatedStaticCameraPerfTestProjectSettings::UAutomatedStaticCameraPerfTestProjectSettings(const FObjectInitializer& Initializer)
 	: Super(Initializer)
@@ -98,6 +99,7 @@ FString UAutomatedStaticCameraPerfTestBase::GetTestID()
 
 bool UAutomatedStaticCameraPerfTestBase::TryStartCSVProfiler(FString CSVFileName)
 {
+#if CSV_PROFILER
 	if(FCsvProfiler* const CsvProfiler = FCsvProfiler::Get())
 	{
 		if (GetCSVOutputMode() != EAutomatedPerfTestCSVOutputMode::Single)
@@ -110,7 +112,7 @@ bool UAutomatedStaticCameraPerfTestBase::TryStartCSVProfiler(FString CSVFileName
 		}
 		return Super::TryStartCSVProfiler(CSVFileName);
 	}
-	
+#endif
 	UE_LOG(LogAutomatedPerfTest, Warning, TEXT("CSVProfiler Start requested, but not available."))
 	return false;
 }
