APT: Tweak the templates to better support Mac/Linux

diff --git a/Engine/Plugins/Performance/AutomatedPerfTesting/Resources/TemplatesForHookup/AutoPerfTests.xml b/Engine/Plugins/Performance/AutomatedPerfTesting/Resources/TemplatesForHookup/AutoPerfTests.xml
--- a/Engine/Plugins/Performance/AutomatedPerfTesting/Resources/TemplatesForHookup/AutoPerfTests.xml
+++ b/Engine/Plugins/Performance/AutomatedPerfTesting/Resources/TemplatesForHookup/AutoPerfTests.xml
@@ -4,6 +4,7 @@
 <!-- Project Info -->
 <Property Name="ProjectName" Value="**REPLACE_PROJECTNAME**" />
 <Property Name="ProjectPath" Value="**REPLACE_PROJECTPATH**" />
+<Option Name="TargetName" DefaultValue="**REPLACE_PROJECTNAME**" Restrict="**REPLACE_PROJECTNAME**" Description="The project target configuration we're building" />
 
 <!-- Stick to the defaults unless the project needs to change this -->
 <Property Name="WithBATDefaults" Value="true" />
@@ -27,6 +28,11 @@
 <!-- Common Setting Overrides -->
 <Property Name="GauntletSettingsFile" Value="$(RootDir)/$(ProjectPath)/Build/Inc/GauntletSettings.xml" />
 
+<!-- For Mac builds, we default to the sandbox entitlement, however PerfTests expect data outside,
+     so we disable the entitlement for this particular build. -->
+<Property Name="ExtraStageAndPackageArguments" Value="-target=&quot;$(TargetName)&quot; -compressed -noEntitlements" />
+<Property Name="ExtraTargetCompileArguments" Value="-noEntitlements"/>
+
 <!-- This will declare an aggregate called BuildAndTest **REPLACE_PROJECTNAME** -->
 <Include Script="$(RootDir)/Engine/Build/Graph/Tasks/BuildAndTestProject.xml" />
 
diff --git a/Engine/Plugins/Performance/AutomatedPerfTesting/Resources/TemplatesForHookup/RunLocalTests.sh b/Engine/Plugins/Performance/AutomatedPerfTesting/Resources/TemplatesForHookup/RunLocalTests.sh
--- a/Engine/Plugins/Performance/AutomatedPerfTesting/Resources/TemplatesForHookup/RunLocalTests.sh
+++ b/Engine/Plugins/Performance/AutomatedPerfTesting/Resources/TemplatesForHookup/RunLocalTests.sh
@@ -1,9 +1,10 @@
+platform=""
 if [ "$(uname)" = "Darwin" ]; then
-	**RUN_UAT_RELATIVEPATH**/RunUAT.sh BuildGraph -Script=**REPLACE_PROJECTPATH**/Build/AutoPerfTests.xml -Target="BuildAndTest **REPLACE_PROJECTNAME**" -UseLocalBuildStorage -set:TargetPlatforms=Mac -set:WithAPTSequenceTests=true -set:IsBuildMachine=true -set:TargetConfigurations=Test -set:DoPerf=true -set:DoLLM=true -set:GenerateLocalReports=true -set:IterationsPerf=1 -set:IterationsLLM=1 -set:IterationsInsights=1
+	platform="Mac"
 fi
 
 if [ "$(uname)" = "Linux" ]; then
-	**RUN_UAT_RELATIVEPATH**/RunUAT.sh BuildGraph -Script=**REPLACE_PROJECTPATH**/Build/AutoPerfTests.xml -Target="BuildAndTest **REPLACE_PROJECTNAME**" -UseLocalBuildStorage -set:TargetPlatforms=Linux -set:WithAPTSequenceTests=true -set:TargetConfigurations=Test -set:DoPerf=true -set:DoLLM=true -set:GenerateLocalReports=true -set:GenerateLocalReports=true -set:IterationsPerf=1 -set:IterationsLLM=1 -set:IterationsInsights=1
+	platform="Linux"
 fi
 
-popd
+**RUN_UAT_RELATIVEPATH**/RunUAT.sh BuildGraph -Script=**REPLACE_PROJECTPATH**/Build/AutoPerfTests.xml -Target="BuildAndTest **REPLACE_PROJECTNAME**" -UseLocalBuildStorage -set:TargetPlatforms=$platform -set:WithAPTSequenceTests=true -set:TargetConfigurations=Test -set:DoPerf=true -set:DoLLM=true -set:GenerateLocalReports=true -set:GenerateLocalReports=true -set:IterationsPerf=1 -set:IterationsLLM=1 -set:IterationsInsights=1
