APT: harden destruction functions.

diff --git a/Engine/Plugins/Performance/AutomatedPerfTesting/Source/AutomatedPerfTesting/Private/AutomatedMaterialPerfTest.cpp b/Engine/Plugins/Performance/AutomatedPerfTesting/Source/AutomatedPerfTesting/Private/AutomatedMaterialPerfTest.cpp
--- a/Engine/Plugins/Performance/AutomatedPerfTesting/Source/AutomatedPerfTesting/Private/AutomatedMaterialPerfTest.cpp
+++ b/Engine/Plugins/Performance/AutomatedPerfTesting/Source/AutomatedPerfTesting/Private/AutomatedMaterialPerfTest.cpp
@@ -255,8 +255,11 @@ void UAutomatedMaterialPerfTest::OnInit()
 void UAutomatedMaterialPerfTest::UnbindAllDelegates()
 {
 	Super::UnbindAllDelegates();
-	
-	GetWorld()->GetTimerManager().ClearAllTimersForObject(this);
+
+	if (UWorld* const World = GetWorld())
+	{
+		World->GetTimerManager().ClearAllTimersForObject(this);
+	}
 }
 
 void UAutomatedMaterialPerfTest::OpenMaterialPerformanceTestMap() const
diff --git a/Engine/Plugins/Performance/AutomatedPerfTesting/Source/AutomatedPerfTesting/Private/AutomatedSequencePerfTest.cpp b/Engine/Plugins/Performance/AutomatedPerfTesting/Source/AutomatedPerfTesting/Private/AutomatedSequencePerfTest.cpp
--- a/Engine/Plugins/Performance/AutomatedPerfTesting/Source/AutomatedPerfTesting/Private/AutomatedSequencePerfTest.cpp
+++ b/Engine/Plugins/Performance/AutomatedPerfTesting/Source/AutomatedPerfTesting/Private/AutomatedSequencePerfTest.cpp
@@ -329,14 +329,22 @@ void UAutomatedSequencePerfTest::UnbindAllDelegates()
 {
 	Super::UnbindAllDelegates();
 
+	UWorld* const World = GetWorld();
+
 	// if we have a valid sequence player, make sure we unbind our events from it when we're wrapping up the test.
-	if(SequencePlayer != nullptr)
+	if (SequencePlayer != nullptr)
 	{
 		SequencePlayer->OnCameraCut.RemoveAll(this);
 		SequencePlayer->OnFinished.RemoveAll(this);
-		GetWorld()->GetTimerManager().ClearAllTimersForObject(SequencePlayer);
+		if (World)
+		{
+			World->GetTimerManager().ClearAllTimersForObject(SequencePlayer);
+		}
 	}
 
 	// clear any stray timers that might be lying around
-	GetWorld()->GetTimerManager().ClearAllTimersForObject(this);
+	if (World)
+	{
+		World->GetTimerManager().ClearAllTimersForObject(this);
+	}
 }
diff --git a/Engine/Plugins/Performance/AutomatedPerfTesting/Source/AutomatedPerfTesting/Private/StaticCameraTests/AutomatedStaticCameraPerfTestBase.cpp b/Engine/Plugins/Performance/AutomatedPerfTesting/Source/AutomatedPerfTesting/Private/StaticCameraTests/AutomatedStaticCameraPerfTestBase.cpp
--- a/Engine/Plugins/Performance/AutomatedPerfTesting/Source/AutomatedPerfTesting/Private/StaticCameraTests/AutomatedStaticCameraPerfTestBase.cpp
+++ b/Engine/Plugins/Performance/AutomatedPerfTesting/Source/AutomatedPerfTesting/Private/StaticCameraTests/AutomatedStaticCameraPerfTestBase.cpp
@@ -314,7 +314,10 @@ void UAutomatedStaticCameraPerfTestBase::UnbindAllDelegates()
 {
 	Super::UnbindAllDelegates();
 	
-	GetWorld()->GetTimerManager().ClearAllTimersForObject(this);
+	if (UWorld* const World = GetWorld())
+	{
+		World->GetTimerManager().ClearAllTimersForObject(this);
+	}
 
 	CurrentCamera = nullptr;
 	CamerasToTest.Empty();
