APT:

Fix APT retriggering SetupTest multiple times
There were two separate issues:
- `OnWorldBeginPlay` fires for each sub level which causes the test to register multiple times
- `SetupTest` would attempt to load a map that could already be loaded and thus re-run `SetupTest`

Fixes are:
- Check if all the Levels in a World are loaded before attempting to run
- Check if the level is already loaded before attempting to reload it

diff --git a/Engine/Plugins/Performance/AutomatedPerfTesting/Source/AutomatedPerfTesting/Private/AutomatedSequencePerfTest.cpp b/Engine/Plugins/Performance/AutomatedPerfTesting/Source/AutomatedPerfTesting/Private/AutomatedSequencePerfTest.cpp
--- a/Engine/Plugins/Performance/AutomatedPerfTesting/Source/AutomatedPerfTesting/Private/AutomatedSequencePerfTest.cpp
+++ b/Engine/Plugins/Performance/AutomatedPerfTesting/Source/AutomatedPerfTesting/Private/AutomatedSequencePerfTest.cpp
@@ -7,6 +7,7 @@
 #include "AutomatedPerfTesting.h"
 #include "LevelSequenceActor.h"
 #include "LevelSequencePlayer.h"
+#include "Engine/LevelStreaming.h"
 #include "MovieSceneSequencePlayer.h"
 #include "Kismet/GameplayStatics.h"
 #include "Misc/CommandLine.h"
@@ -43,14 +44,20 @@ void UAutomatedSequencePerfTest::SetupTest()
 {
 	if(CurrentMapSequenceCombo.IsSet())
 	{
+		// don't even try to set up the test if we're not in the correct map
 		if(GetCurrentMap() == CurrentMapSequenceCombo->Map.GetAssetName())
 		{
-			// don't even try to set up the test if we're not in the correct map
-			Super::SetupTest();
-			
 			// make sure the world exists, then create a sequence player
-			if(UWorld* const World = GetWorld())
+			if (UWorld* const World = GetWorld())
 			{
+				// Begin World Fires for each sub level so we need to make sure that all of the required ones are ready to go
+				// TODO: Consider pushing this logic into the Base Class
+				if (World->GetNumStreamingLevelsBeingLoaded() != 0)
+				{
+					return;
+				}
+			
+				Super::SetupTest();
 				// reset the camera cut number
 				NumCameraCuts = -1;
 				
@@ -119,8 +126,16 @@ void UAutomatedSequencePerfTest::NextMap()
 			OptionsString += "game=" + CurrentMapSequenceCombo->GameModeOverride;
 		}
 		
-		UE_LOG(LogAutomatedPerfTest, Log, TEXT("Opening map %s%s"), *CurrentMapSequenceCombo->Map.GetAssetName(), *OptionsString);
-		UGameplayStatics::OpenLevel(AutomatedPerfTest::FindCurrentWorld(), *CurrentMapSequenceCombo->Map.GetAssetName(), true, OptionsString);
+		if (OptionsString.IsEmpty() && AutomatedPerfTest::FindCurrentWorld()->GetName() == CurrentMapSequenceCombo->Map.GetAssetName())
+		{
+			UE_LOG(LogAutomatedPerfTest, Log, TEXT("%s is already loaded and does not have any options string, skipping map load and setting up the test"), *CurrentMapSequenceCombo->Map.GetAssetName());
+			SetupTest();
+		}
+		else 
+		{
+			UE_LOG(LogAutomatedPerfTest, Log, TEXT("Opening map %s%s"), *CurrentMapSequenceCombo->Map.GetAssetName(), *OptionsString);
+			UGameplayStatics::OpenLevel(AutomatedPerfTest::FindCurrentWorld(), *CurrentMapSequenceCombo->Map.GetAssetName(), true, OptionsString);
+		}
 	}
 	else
 	{
