APT: Support Deferred Replays

* There are instances where Replays could be deferred within a GameInstance implementation. This allows APT to know if a replay is deferred
* Split up profiling setup and markers to separate functions for further flexibility
* Only start perf metric collection once Replay playback is ready

#rnx

diff --git a/Engine/Plugins/Performance/AutomatedPerfTesting/Source/AutomatedPerfTesting/Private/AutomatedMaterialPerfTest.cpp b/Engine/Plugins/Performance/AutomatedPerfTesting/Source/AutomatedPerfTesting/Private/AutomatedMaterialPerfTest.cpp
--- a/Engine/Plugins/Performance/AutomatedPerfTesting/Source/AutomatedPerfTesting/Private/AutomatedMaterialPerfTest.cpp
+++ b/Engine/Plugins/Performance/AutomatedPerfTesting/Source/AutomatedPerfTesting/Private/AutomatedMaterialPerfTest.cpp
@@ -38,6 +38,8 @@ void UAutomatedMaterialPerfTest::SetupTest()
 		// make sure the world exists, then create a sequence player
 		if(UWorld* const World = GetWorld())
 		{
+			SetupProfiling();
+
 			// Also load and verify the material plate mesh
 			UE_LOG(LogAutomatedPerfTest, Log, TEXT("Loading material plate mesh: %s"), *Settings->MaterialPlate.ToString())
 			UStaticMesh* LoadedMaterialPlateMesh = LoadObject<UStaticMesh>(NULL, *Settings->MaterialPlate.ToString(), NULL, LOAD_None, NULL);
@@ -89,6 +91,9 @@ void UAutomatedMaterialPerfTest::RunTest()
 {
 	Super::RunTest();
 
+	// The overall profiling region starts here, but each material gets its own sub region for easy delineation. 
+	MarkProfilingStart(); 
+
 	UE_LOG(LogAutomatedPerfTest, Log, TEXT("UAutomatedMaterialPerfTest::RunTest"))
 	
 	FTimerHandle UnusedHandle;
@@ -262,6 +267,14 @@ void UAutomatedMaterialPerfTest::UnbindAllDelegates()
 	}
 }
 
+void UAutomatedMaterialPerfTest::TeardownTest(bool bExitAfterTeardown)
+{
+	UE_LOG(LogAutomatedPerfTest, Log, TEXT("AutomatedMaterialPerfTest::TeardownTest"));
+	MarkProfilingEnd();
+	TeardownProfiling();
+	Super::TeardownTest(bExitAfterTeardown);
+}
+
 void UAutomatedMaterialPerfTest::OpenMaterialPerformanceTestMap() const
 {
 	// no need to prepend this with a ? since OpenLevel handles that part for us
diff --git a/Engine/Plugins/Performance/AutomatedPerfTesting/Source/AutomatedPerfTesting/Private/AutomatedPerfTestControllerBase.cpp b/Engine/Plugins/Performance/AutomatedPerfTesting/Source/AutomatedPerfTesting/Private/AutomatedPerfTestControllerBase.cpp
--- a/Engine/Plugins/Performance/AutomatedPerfTesting/Source/AutomatedPerfTesting/Private/AutomatedPerfTestControllerBase.cpp
+++ b/Engine/Plugins/Performance/AutomatedPerfTesting/Source/AutomatedPerfTesting/Private/AutomatedPerfTestControllerBase.cpp
@@ -115,6 +115,7 @@ UAutomatedPerfTestControllerBase::UAutomatedPerfTestControllerBase(const FObject
 	, bRequestsCSVProfiler(false)
 	, bRequestsVideoCapture(false)
 	, bRequestsLockedDynRes(false)
+	, InsightsRegionID(0)
     , CSVOutputMode(EAutomatedPerfTestCSVOutputMode::Single)
 {
 	// cache this off once, so that it's consistent throughout a session
@@ -234,8 +235,11 @@ bool UAutomatedPerfTestControllerBase::TryStopCSVProfiler()
 #if CSV_PROFILER
 	if(FCsvProfiler* const CsvProfiler = FCsvProfiler::Get())
 	{
-		const FGraphEventRef AutomatedPerfTestEndEvent = FGraphEvent::CreateGraphEvent();
-		CsvProfiler->EndCapture(AutomatedPerfTestEndEvent);
+		if (CsvProfiler->IsCapturing())
+		{
+			const FGraphEventRef AutomatedPerfTestEndEvent = FGraphEvent::CreateGraphEvent();
+			CsvProfiler->EndCapture(AutomatedPerfTestEndEvent);
+		}
 		return true;
 	}
 #endif
@@ -334,26 +338,7 @@ bool UAutomatedPerfTestControllerBase::TryFinalizingVideoCapture(const bool bSto
 void UAutomatedPerfTestControllerBase::SetupTest()
 {
 	UE_LOG(LogAutomatedPerfTest, Log, TEXT("Base:: SetupTest"));
-	if(RequestsInsightsTrace())
-	{
-		TryStartInsightsTrace();
-	}
-
-	if(RequestsCSVProfiler() && CSVOutputMode ==  EAutomatedPerfTestCSVOutputMode::Single)
-	{
-		TryStartCSVProfiler();
-	}
-
-	if(RequestsFPSChart())
-	{
-		TryStartFPSChart();
-	}
 
-	if(RequestsVideoCapture())
-	{
-		TryStartVideoCapture();
-	}
-	
 	SetupGameModeInstance();
 
 	// Subclasses should implement their own transitions from SetupTest to RunTest depending on their needs
@@ -364,20 +349,7 @@ void UAutomatedPerfTestControllerBase::SetupTest()
 void UAutomatedPerfTestControllerBase::RunTest()
 {
 	UE_LOG(LogAutomatedPerfTest, Log, TEXT("Base:: RunTest"));
-	if(RequestsInsightsTrace())
-	{
-		TRACE_BEGIN_REGION(*GetOverallRegionName());
-	}
 
-	if(RequestsCSVProfiler())
-	{
-		if(CSVOutputMode == EAutomatedPerfTestCSVOutputMode::Separate)
-		{
-			TryStartCSVProfiler();
-		}
-		CSV_EVENT(AutomatedPerfTest, TEXT("START"), *GetOverallRegionName())
-	}
-	
 	if(GameMode && GameMode->GetClass()->ImplementsInterface(UAutomatedPerfTestInterface::StaticClass()))
 	{
 		IAutomatedPerfTestInterface::Execute_RunTest(GameMode);
@@ -387,33 +359,6 @@ void UAutomatedPerfTestControllerBase::RunTest()
 void UAutomatedPerfTestControllerBase::TeardownTest(bool bExitAfterTeardown)
 {
 	UE_LOG(LogAutomatedPerfTest, Log, TEXT("Base:: TeardownTest"));
-	
-	if(RequestsInsightsTrace())
-	{
-		TRACE_END_REGION(*GetOverallRegionName());
-	}
-	
-	if(RequestsCSVProfiler())
-	{
-		CSV_EVENT(AutomatedPerfTest, TEXT("END"), *GetOverallRegionName())
-		if(CSVOutputMode == EAutomatedPerfTestCSVOutputMode::Separate)
-		{
-			TryStopCSVProfiler();
-		}
-	}
-
-	if(RequestsFPSChart())
-	{
-		TryStopFPSChart();
-	}
-
-	if(RequestsVideoCapture())
-	{
-		if(!TryFinalizingVideoCapture())
-		{
-			UE_LOG(LogAutomatedPerfTest, Warning, TEXT("Attempted to finalize requested video capture, but failed."))
-		}
-	}
 
 	if(GameMode && GameMode->GetClass()->ImplementsInterface(UAutomatedPerfTestInterface::StaticClass()))
 	{
@@ -435,20 +380,6 @@ void UAutomatedPerfTestControllerBase::TriggerExitAfterDelay()
 
 void UAutomatedPerfTestControllerBase::Exit()
 {
-	if(RequestsInsightsTrace())
-	{
-		TryStopInsightsTrace();
-	}
-
-	if(RequestsCSVProfiler())
-	{
-		CSV_EVENT(AutomatedPerfTest, TEXT("END"), *GetOverallRegionName())
-		if(CSVOutputMode == EAutomatedPerfTestCSVOutputMode::Single)
-		{
-			TryStopCSVProfiler();
-		}
-	}
-
 	if(GameMode && GameMode->GetClass()->ImplementsInterface(UAutomatedPerfTestInterface::StaticClass()))
 	{
 		IAutomatedPerfTestInterface::Execute_Exit(GameMode);
@@ -506,6 +437,103 @@ void UAutomatedPerfTestControllerBase::SetupGameModeInstance()
 	}
 }
 
+void UAutomatedPerfTestControllerBase::SetupProfiling()
+{
+
+	if (RequestsCSVProfiler() && CSVOutputMode == EAutomatedPerfTestCSVOutputMode::Single)
+	{
+		TryStartCSVProfiler();
+	}
+
+	if (RequestsFPSChart())
+	{
+		TryStartFPSChart();
+	}
+
+	if (RequestsVideoCapture())
+	{
+		TryStartVideoCapture();
+	}
+}
+
+void UAutomatedPerfTestControllerBase::InitializeInsights()
+{
+	if (RequestsInsightsTrace())
+	{
+		TryStartInsightsTrace();
+	}
+}
+
+void UAutomatedPerfTestControllerBase::ShutdownInsights()
+{
+	if (RequestsInsightsTrace())
+	{
+		TryStopInsightsTrace();
+	}
+}
+
+void UAutomatedPerfTestControllerBase::MarkProfilingStart()
+{
+	if (RequestsInsightsTrace())
+	{
+		InsightsRegionID = TRACE_BEGIN_REGION_WITH_ID(*GetOverallRegionName(), TEXT("AutomatedPerfTest"));
+	}
+
+	if (RequestsCSVProfiler())
+	{
+		if (CSVOutputMode == EAutomatedPerfTestCSVOutputMode::Separate)
+		{
+			TryStartCSVProfiler();
+		}
+		CSV_EVENT(AutomatedPerfTest, TEXT("START"), *GetOverallRegionName())
+	}
+}
+
+void UAutomatedPerfTestControllerBase::TeardownProfiling()
+{
+	if (RequestsCSVProfiler())
+	{
+		CSV_EVENT(AutomatedPerfTest, TEXT("END"), *GetOverallRegionName())
+		if (CSVOutputMode == EAutomatedPerfTestCSVOutputMode::Single)
+		{
+			TryStopCSVProfiler();
+		}
+	}
+
+	if (RequestsFPSChart())
+	{
+		TryStopFPSChart();
+	}
+
+	if (RequestsVideoCapture())
+	{
+		if (!TryFinalizingVideoCapture())
+		{
+			UE_LOG(LogAutomatedPerfTest, Warning, TEXT("Attempted to finalize requested video capture, but failed."))
+		}
+	}
+}
+
+void UAutomatedPerfTestControllerBase::MarkProfilingEnd()
+{
+	if (RequestsInsightsTrace())
+	{
+		if (InsightsRegionID != 0)
+		{
+			TRACE_END_REGION_WITH_ID(InsightsRegionID);
+		}
+	}
+
+	if (RequestsCSVProfiler())
+	{
+		CSV_EVENT(AutomatedPerfTest, TEXT("END"), *GetOverallRegionName())
+		if (CSVOutputMode == EAutomatedPerfTestCSVOutputMode::Separate)
+		{
+			TryStopCSVProfiler();
+		}
+	}
+}
+
 void UAutomatedPerfTestControllerBase::SetCSVOutputMode(EAutomatedPerfTestCSVOutputMode NewOutputMode)
 {
 	CSVOutputMode = NewOutputMode;	
@@ -550,6 +578,8 @@ void UAutomatedPerfTestControllerBase::OnInit()
 	}
 	
 	FWorldDelegates::OnPreWorldInitialization.AddUObject(this, &ThisClass::OnPreWorldInitializeInternal);
+
+	InitializeInsights();
 }
 
 void UAutomatedPerfTestControllerBase::OnTick(float TimeDelta)
@@ -579,6 +609,7 @@ void UAutomatedPerfTestControllerBase::BeginDestroy()
 void UAutomatedPerfTestControllerBase::EndAutomatedPerfTest(const int32 ExitCode)
 {
 	UnbindAllDelegates();
+	ShutdownInsights();
 
 	UE_LOG(LogAutomatedPerfTest, Log, TEXT("Test ID %s completed, requesting exit..."), *GetTestID());
 	
diff --git a/Engine/Plugins/Performance/AutomatedPerfTesting/Source/AutomatedPerfTesting/Private/AutomatedReplayPerfTest.cpp b/Engine/Plugins/Performance/AutomatedPerfTesting/Source/AutomatedPerfTesting/Private/AutomatedReplayPerfTest.cpp
--- a/Engine/Plugins/Performance/AutomatedPerfTesting/Source/AutomatedPerfTesting/Private/AutomatedReplayPerfTest.cpp
+++ b/Engine/Plugins/Performance/AutomatedPerfTesting/Source/AutomatedPerfTesting/Private/AutomatedReplayPerfTest.cpp
@@ -53,6 +53,8 @@ FString UAutomatedReplayPerfTest::GetTestID()
 
 void UAutomatedReplayPerfTest::SetupTest()
 {
+	bIsTearingDown = false;
+
 	if (UWorld* World = GetWorld())
 	{
 		// If this is true, it usually means replay is already 
@@ -67,10 +69,9 @@ void UAutomatedReplayPerfTest::SetupTest()
 			return;
 		}
 
-		Super::SetupTest();
-
-		// Register delegates to be called when replay playback is 
-		// complete or if there is a failure for some reason.
+		// Register delegates to be called when replay playback has started or 
+		// completed or if there is a failure for some reason.
+		FNetworkReplayDelegates::OnReplayStarted.AddUObject(this, &ThisClass::OnReplayStarted);
 		FNetworkReplayDelegates::OnReplayPlaybackComplete.AddUObject(this, &ThisClass::OnReplayComplete);
 		FNetworkReplayDelegates::OnReplayPlaybackFailure.AddUObject(this, &ThisClass::OnReplayFailure);
 
@@ -86,13 +87,14 @@ void UAutomatedReplayPerfTest::SetupTest()
 
 void UAutomatedReplayPerfTest::RunTest()
 {
-	Super::RunTest();
-
 	UWorld* World = GetWorld();
 	UGameInstance* GameInstance = World ? World->GetGameInstance() : nullptr;
 	if (GameInstance)
 	{
 		bIsReplayTriggered = GameInstance->PlayReplay(ReplayName);
+
+		// Also check if replay playback was deferred for any reason. 
+		bIsReplayTriggered = bIsReplayTriggered || GameInstance->IsReplayDeferred(ReplayName);
 	}
 
 	if (!bIsReplayTriggered)
@@ -104,6 +106,18 @@ void UAutomatedReplayPerfTest::RunTest()
 
 void UAutomatedReplayPerfTest::TeardownTest(bool bExitAfterTeardown)
 {
+	if (bIsTearingDown)
+	{
+		// When a replay completes, it can call OnReplayPlaybackComplete multiple times until
+		// exit, causing Teardown to be invoked multiple times. This mitigates that issue.
+		return;
+	}
+
+	UE_LOG(LogAutomatedPerfTest, Log, TEXT("AutomatedReplayPerfTest::TeardownTest"));
+
+	bIsTearingDown = true;
+	MarkProfilingEnd();
+	TeardownProfiling();
 	Super::TeardownTest(bExitAfterTeardown);
 }
 
@@ -179,16 +193,27 @@ void UAutomatedReplayPerfTest::OnInit()
 
 void UAutomatedReplayPerfTest::UnbindAllDelegates()
 {
+	FNetworkReplayDelegates::OnReplayStarted.RemoveAll(this);
 	FNetworkReplayDelegates::OnReplayPlaybackComplete.RemoveAll(this);
 	FNetworkReplayDelegates::OnReplayPlaybackFailure.RemoveAll(this);
 	Super::UnbindAllDelegates();
 }
 
+void UAutomatedReplayPerfTest::OnReplayStarted(UWorld* World)
+{
+	// Set up Insights/CSV profiler and mark the start events/regions. 
+	// We do this here so that metric collection starts only when replay 
+	// playback is ready. This is especially useful since Replays can
+	// be deferred.
+	SetupProfiling();
+	MarkProfilingStart();
+}
+
 void UAutomatedReplayPerfTest::OnReplayComplete(UWorld* World)
 {
 	// We manually ensure we tear down and exit here the moment replay
 	// playback is completed, otherwise it requires user input to exit. 
-	constexpr bool bDelayedExitAfterTeardown = true; 
+	constexpr bool bDelayedExitAfterTeardown = true;
 	TeardownTest(bDelayedExitAfterTeardown);
 }
 
diff --git a/Engine/Plugins/Performance/AutomatedPerfTesting/Source/AutomatedPerfTesting/Private/AutomatedSequencePerfTest.cpp b/Engine/Plugins/Performance/AutomatedPerfTesting/Source/AutomatedPerfTesting/Private/AutomatedSequencePerfTest.cpp
--- a/Engine/Plugins/Performance/AutomatedPerfTesting/Source/AutomatedPerfTesting/Private/AutomatedSequencePerfTest.cpp
+++ b/Engine/Plugins/Performance/AutomatedPerfTesting/Source/AutomatedPerfTesting/Private/AutomatedSequencePerfTest.cpp
@@ -85,6 +85,9 @@ void UAutomatedSequencePerfTest::SetupTest()
 				UE_LOG(LogAutomatedPerfTest, Log, TEXT("SetupMapTest:: Scrubbing to start"));
 				SequencePlayer->PlayTo(PlaybackParams, PlayToParams);
 
+				// Setup profiling only when we are sure to run the sequence. 
+				SetupProfiling();
+
 				UE_LOG(LogAutomatedPerfTest, Verbose, TEXT("SetupMapTest:: Waiting for %f seconds before playing seqeuence"), Settings->SequenceStartDelay);
 				FTimerHandle UnusedHandle;
 				GetWorld()->GetTimerManager().SetTimer(UnusedHandle, this, &UAutomatedSequencePerfTest::RunTest, 1.0, false, Settings->SequenceStartDelay);
@@ -154,6 +157,7 @@ void UAutomatedSequencePerfTest::RunTest()
 	if(SequencePlayer)
 	{
 		UE_LOG(LogAutomatedPerfTest, Log, TEXT("RunTest::Valid Sequence Player, proceeding"));
+		MarkProfilingStart(); // Mark profiler start region on sequence start. 
 
 		// trigger a camera cut manually in order to start the region for the first camera cut
         SequencePlayer->Play();
@@ -185,6 +189,8 @@ void UAutomatedSequencePerfTest::TeardownTest(bool bExitAfterTeardown)
 {
 	UE_LOG(LogAutomatedPerfTest, Log, TEXT("AutomatedSequencePerfTest::TeardownTest"));
 	
+	MarkProfilingEnd();
+	TeardownProfiling();
 	Super::TeardownTest(bExitAfterTeardown);
 
 	UnbindAllDelegates();
diff --git a/Engine/Plugins/Performance/AutomatedPerfTesting/Source/AutomatedPerfTesting/Private/StaticCameraTests/AutomatedStaticCameraPerfTestBase.cpp b/Engine/Plugins/Performance/AutomatedPerfTesting/Source/AutomatedPerfTesting/Private/StaticCameraTests/AutomatedStaticCameraPerfTestBase.cpp
--- a/Engine/Plugins/Performance/AutomatedPerfTesting/Source/AutomatedPerfTesting/Private/StaticCameraTests/AutomatedStaticCameraPerfTestBase.cpp
+++ b/Engine/Plugins/Performance/AutomatedPerfTesting/Source/AutomatedPerfTesting/Private/StaticCameraTests/AutomatedStaticCameraPerfTestBase.cpp
@@ -40,7 +40,6 @@ void UAutomatedStaticCameraPerfTestBase::SetupTest()
 		if(GetCurrentMap() == CurrentMapPath.GetAssetName())
 		{
 			Super::SetupTest();
-			
 			UE_LOG(LogAutomatedPerfTest, Log, TEXT("UAutomatedStaticCameraPerfTestBase::SetupTest"))
 			// make sure the world exists, then create a sequence player
 			if(UWorld* const World = GetWorld())
@@ -57,6 +56,9 @@ void UAutomatedStaticCameraPerfTestBase::SetupTest()
 				UE_LOG(LogAutomatedPerfTest, Log, TEXT("Found %hhd cameras to test in map %s"), CamerasToTest.Num(), *CurrentMapPath.GetAssetName());
 
 				GetFirstPlayerController()->SetCinematicMode(true, true, true);
+
+				// Setup profiling once we have loaded the map. 
+				SetupProfiling();
 				
 				// delay for WarmUpDelay, and call RunTest
 				FTimerHandle UnusedHandle;
@@ -87,6 +89,7 @@ void UAutomatedStaticCameraPerfTestBase::RunTest()
 	Super::RunTest();
 
 	UE_LOG(LogAutomatedPerfTest, Log, TEXT("UAutomatedStaticCameraPerfTestBase::RunTest"))
+	MarkProfilingStart(); // Mark the start region on the profiler once the test starts
 	
 	FTimerHandle UnusedHandle;
 	GetWorld()->GetTimerManager().SetTimer(UnusedHandle, this, &UAutomatedStaticCameraPerfTestBase::SetUpNextCamera, 1.0, false, Settings->WarmUpTime);
@@ -279,7 +282,7 @@ void UAutomatedStaticCameraPerfTestBase::OnInit()
 	{
 		for (FSoftObjectPath MapPath : Settings->MapsToTest)
 		{
-			if(MapPath.GetAssetName() == CurrentMapName)
+			if (MapPath.GetAssetName() == CurrentMapName)
 			{
 				MapsToTest.Add(MapPath);
 			}
@@ -321,4 +324,11 @@ void UAutomatedStaticCameraPerfTestBase::UnbindAllDelegates()
 
 	CurrentCamera = nullptr;
 	CamerasToTest.Empty();
-}
\ No newline at end of file
+}
+
+void UAutomatedStaticCameraPerfTestBase::TriggerExitAfterDelay()
+{
+	MarkProfilingEnd();
+	TeardownProfiling();
+	Super::TriggerExitAfterDelay();
+}
diff --git a/Engine/Plugins/Performance/AutomatedPerfTesting/Source/AutomatedPerfTesting/Public/AutomatedMaterialPerfTest.h b/Engine/Plugins/Performance/AutomatedPerfTesting/Source/AutomatedPerfTesting/Public/AutomatedMaterialPerfTest.h
--- a/Engine/Plugins/Performance/AutomatedPerfTesting/Source/AutomatedPerfTesting/Public/AutomatedMaterialPerfTest.h
+++ b/Engine/Plugins/Performance/AutomatedPerfTesting/Source/AutomatedPerfTesting/Public/AutomatedMaterialPerfTest.h
@@ -140,6 +140,7 @@ public:
 protected:
 	UE_API virtual void OnInit() override;
 	UE_API virtual void UnbindAllDelegates() override;
+	UE_API virtual void TeardownTest(bool bExitAfterTeardown = true) override;
 
 	UE_API void OpenMaterialPerformanceTestMap() const;
 
diff --git a/Engine/Plugins/Performance/AutomatedPerfTesting/Source/AutomatedPerfTesting/Public/AutomatedPerfTestControllerBase.h b/Engine/Plugins/Performance/AutomatedPerfTesting/Source/AutomatedPerfTesting/Public/AutomatedPerfTestControllerBase.h
--- a/Engine/Plugins/Performance/AutomatedPerfTesting/Source/AutomatedPerfTesting/Public/AutomatedPerfTestControllerBase.h
+++ b/Engine/Plugins/Performance/AutomatedPerfTesting/Source/AutomatedPerfTesting/Public/AutomatedPerfTestControllerBase.h
@@ -130,6 +130,15 @@ protected:
 	UE_API virtual void UnbindAllDelegates();
 
 	UE_API void SetupGameModeInstance();
+	
+	/* Profiling Functions */
+	UE_API void SetupProfiling();
+	UE_API void InitializeInsights();
+	UE_API void ShutdownInsights();
+	UE_API void TeardownProfiling();
+	UE_API void MarkProfilingStart();
+	UE_API void MarkProfilingEnd();
+	/* End Profiling Functions*/
 
 private:
 	FString TraceChannels;
@@ -141,6 +150,7 @@ private:
 	bool bRequestsCSVProfiler;
 	bool bRequestsVideoCapture;
 	bool bRequestsLockedDynRes;
+	uint64 InsightsRegionID;
 
 	FText VideoRecordingTitle;
 	
diff --git a/Engine/Plugins/Performance/AutomatedPerfTesting/Source/AutomatedPerfTesting/Public/AutomatedReplayPerfTest.h b/Engine/Plugins/Performance/AutomatedPerfTesting/Source/AutomatedPerfTesting/Public/AutomatedReplayPerfTest.h
--- a/Engine/Plugins/Performance/AutomatedPerfTesting/Source/AutomatedPerfTesting/Public/AutomatedReplayPerfTest.h
+++ b/Engine/Plugins/Performance/AutomatedPerfTesting/Source/AutomatedPerfTesting/Public/AutomatedReplayPerfTest.h
@@ -63,12 +63,14 @@ protected:
 	UE_API virtual void OnInit() override;
 	UE_API virtual void UnbindAllDelegates() override;
 
+	UE_API void OnReplayStarted(UWorld* World);
 	UE_API void OnReplayComplete(UWorld* World);
 	UE_API void OnReplayFailure(UWorld* World, const UE::Net::TNetResult<EReplayResult>& Error);
 
 private:
 	FString ReplayName;
-	bool	bIsReplayTriggered = false;
+	bool	bIsReplayTriggered	= false;
+	bool	bIsTearingDown		= false;
 };
 
 #undef UE_API
diff --git a/Engine/Plugins/Performance/AutomatedPerfTesting/Source/AutomatedPerfTesting/Public/StaticCameraTests/AutomatedStaticCameraPerfTestBase.h b/Engine/Plugins/Performance/AutomatedPerfTesting/Source/AutomatedPerfTesting/Public/StaticCameraTests/AutomatedStaticCameraPerfTestBase.h
--- a/Engine/Plugins/Performance/AutomatedPerfTesting/Source/AutomatedPerfTesting/Public/StaticCameraTests/AutomatedStaticCameraPerfTestBase.h
+++ b/Engine/Plugins/Performance/AutomatedPerfTesting/Source/AutomatedPerfTesting/Public/StaticCameraTests/AutomatedStaticCameraPerfTestBase.h
@@ -118,6 +118,7 @@ public:
 protected:
 	UE_API virtual void OnInit() override;
 	UE_API virtual void UnbindAllDelegates() override;
+	UE_API virtual void TriggerExitAfterDelay() override;
 
 private:
 	UPROPERTY()
diff --git a/Engine/Source/Runtime/Engine/Classes/Engine/GameInstance.h b/Engine/Source/Runtime/Engine/Classes/Engine/GameInstance.h
--- a/Engine/Source/Runtime/Engine/Classes/Engine/GameInstance.h
+++ b/Engine/Source/Runtime/Engine/Classes/Engine/GameInstance.h
@@ -515,6 +515,14 @@ public:
 	 */
 	ENGINE_API virtual bool PlayReplay(const FString& InName, UWorld* WorldOverride = nullptr, const TArray<FString>& AdditionalOptions = TArray<FString>());
 
+	/**
+	 * If given replay playback was deferred and has not started playback yet, this will return true.
+	 * 
+	 * @param ReplayName Name of the replay file.
+	 * @return True if replay playback was deferred
+	 */
+	ENGINE_API virtual bool IsReplayDeferred(const FString& ReplayName) const;
+
 	/**
 	 * Start playing back a playlist of previously recorded replays.
 	 *
diff --git a/Engine/Source/Runtime/Engine/Private/GameInstance.cpp b/Engine/Source/Runtime/Engine/Private/GameInstance.cpp
--- a/Engine/Source/Runtime/Engine/Private/GameInstance.cpp
+++ b/Engine/Source/Runtime/Engine/Private/GameInstance.cpp
@@ -1273,6 +1273,16 @@ bool UGameInstance::PlayReplay(const FString& Name, UWorld* WorldOverride, const
 	return false;
 }
 
+bool UGameInstance::IsReplayDeferred(const FString& ReplayName) const
+{
+	if (UReplaySubsystem* ReplaySubsystem = GetSubsystem<UReplaySubsystem>())
+	{
+		return ReplaySubsystem->IsReplayDeferred(ReplayName);
+	}
+
+	return false;
+}
+
 class FGameInstanceReplayPlaylistHelper
 {
 private:
diff --git a/Engine/Source/Runtime/Engine/Private/ReplaySubsystem.cpp b/Engine/Source/Runtime/Engine/Private/ReplaySubsystem.cpp
--- a/Engine/Source/Runtime/Engine/Private/ReplaySubsystem.cpp
+++ b/Engine/Source/Runtime/Engine/Private/ReplaySubsystem.cpp
@@ -460,6 +460,12 @@ bool UReplaySubsystem::IsSavingCheckpoint() const
 	return false;
 }
 
+bool UReplaySubsystem::IsReplayDeferred(const FString& ReplayName) const
+{
+	// Default implementation of replay subsystem does not defer replays 
+	return false;
+}
+
 void UReplaySubsystem::SetCheckpointSaveMaxMSPerFrame(const float InCheckpointSaveMaxMSPerFrame)
 {
 	UWorld* CurrentWorld = GetWorld();
diff --git a/Engine/Source/Runtime/Engine/Public/ReplaySubsystem.h b/Engine/Source/Runtime/Engine/Public/ReplaySubsystem.h
--- a/Engine/Source/Runtime/Engine/Public/ReplaySubsystem.h
+++ b/Engine/Source/Runtime/Engine/Public/ReplaySubsystem.h
@@ -92,6 +92,14 @@ public:
 	
 	ENGINE_API bool IsSavingCheckpoint() const;
 
+	/**
+	 * If given replay playback was deferred and has not started playback yet, this will return true.
+	 *
+	 * @param ReplayName Name of the replay file.
+	 * @return True if replay playback was deferred
+	 */
+	ENGINE_API virtual bool IsReplayDeferred(const FString& ReplayName) const;
+
 	/**
 	 * Add an event to the currently recording replay, associated with the current time
 	 *
