APT: fix a build farm warning.

diff --git a/Engine/Plugins/Performance/AutomatedPerfTesting/Resources/TemplatesForHookup/AutoPerfTests.xml b/Engine/Plugins/Performance/AutomatedPerfTesting/Resources/TemplatesForHookup/AutoPerfTests.xml
--- a/Engine/Plugins/Performance/AutomatedPerfTesting/Resources/TemplatesForHookup/AutoPerfTests.xml
+++ b/Engine/Plugins/Performance/AutomatedPerfTesting/Resources/TemplatesForHookup/AutoPerfTests.xml
@@ -28,12 +28,12 @@
 <!-- Common Setting Overrides -->
 <Property Name="GauntletSettingsFile" Value="$(RootDir)/$(ProjectPath)/Build/Inc/GauntletSettings.xml" />
 
-<!-- For Mac builds, we default to the sandbox entitlement, however PerfTests expect data outside,
-     so we disable the entitlement for this particular build. -->
-<Property Name="ExtraStageAndPackageArguments" Value="-target=&quot;$(TargetName)&quot; -compressed -noEntitlements" />
-<Property Name="ExtraTargetCompileArguments" Value="-noEntitlements"/>
-
 <!-- This will declare an aggregate called BuildAndTest **REPLACE_PROJECTNAME** -->
 <Include Script="$(RootDir)/Engine/Build/Graph/Tasks/BuildAndTestProject.xml" />
 
+<!-- For Mac builds, we default to the sandbox entitlement, however PerfTests expect data outside,
+     so we disable the entitlement for this particular build. -->
+<Property Name="ExtraStageAndPackageArguments" Value="$(ExtraStageAndPackageArguments) -target=&quot;$(TargetName)&quot; -compressed -noEntitlements" />
+<Property Name="ExtraTargetCompileArguments" Value="$(ExtraTargetCompileArguments) -noEntitlements"/>
+
 </BuildGraph>
