Fix running APT'd projects on a Mac.

- The previous fix for a build farm warning (UE-278581) wasn't actually a workable solution for a Mac.
- Current solution is to use a new option created for the Entitlements specifically which isn't being overridden in Horde jobs, thus avoiding the warning.

diff --git a/Engine/Plugins/Performance/AutomatedPerfTesting/Resources/TemplatesForHookup/AutoPerfTests.xml b/Engine/Plugins/Performance/AutomatedPerfTesting/Resources/TemplatesForHookup/AutoPerfTests.xml
--- a/Engine/Plugins/Performance/AutomatedPerfTesting/Resources/TemplatesForHookup/AutoPerfTests.xml
+++ b/Engine/Plugins/Performance/AutomatedPerfTesting/Resources/TemplatesForHookup/AutoPerfTests.xml
@@ -28,12 +28,12 @@
 <!-- Common Setting Overrides -->
 <Property Name="GauntletSettingsFile" Value="$(RootDir)/$(ProjectPath)/Build/Inc/GauntletSettings.xml" />
 
-<!-- This will declare an aggregate called BuildAndTest **REPLACE_PROJECTNAME** -->
-<Include Script="$(RootDir)/Engine/Build/Graph/Tasks/BuildAndTestProject.xml" />
-
 <!-- For Mac builds, we default to the sandbox entitlement, however PerfTests expect data outside,
      so we disable the entitlement for this particular build. -->
-<Property Name="ExtraStageAndPackageArguments" Value="$(ExtraStageAndPackageArguments) -target=&quot;$(TargetName)&quot; -compressed -noEntitlements" />
-<Property Name="ExtraTargetCompileArguments" Value="$(ExtraTargetCompileArguments) -noEntitlements"/>
+<Property Name="ExtraStageAndPackageArguments" Value="-target=&quot;$(TargetName)&quot; -compressed -noEntitlements" />
+<Property Name="ExtraTargetEntitlementArguments" Value="-noEntitlements"/>
+
+<!-- This will declare an aggregate called BuildAndTest **REPLACE_PROJECTNAME** -->
+<Include Script="$(RootDir)/Engine/Build/Graph/Tasks/BuildAndTestProject.xml" />
 
 </BuildGraph>
