Fix crash in APT by checking if an end capture command has been already enqueued to the CsvProfiler and avoiding creating a superfluous task graph event that would never be signalled

#rnx

diff --git a/Engine/Plugins/Performance/AutomatedPerfTesting/Source/AutomatedPerfTesting/Private/AutomatedPerfTestControllerBase.cpp b/Engine/Plugins/Performance/AutomatedPerfTesting/Source/AutomatedPerfTesting/Private/AutomatedPerfTestControllerBase.cpp
--- a/Engine/Plugins/Performance/AutomatedPerfTesting/Source/AutomatedPerfTesting/Private/AutomatedPerfTestControllerBase.cpp
+++ b/Engine/Plugins/Performance/AutomatedPerfTesting/Source/AutomatedPerfTesting/Private/AutomatedPerfTestControllerBase.cpp
@@ -235,7 +235,8 @@ bool UAutomatedPerfTestControllerBase::TryStopCSVProfiler()
 #if CSV_PROFILER
 	if(FCsvProfiler* const CsvProfiler = FCsvProfiler::Get())
 	{
-		if (CsvProfiler->IsCapturing())
+		const bool bSafeToStopCsvProfiler = CsvProfiler->IsCapturing() && !CsvProfiler->IsEndCapturePending();
+		if (bSafeToStopCsvProfiler)
 		{
 			const FGraphEventRef AutomatedPerfTestEndEvent = FGraphEvent::CreateGraphEvent();
 			CsvProfiler->EndCapture(AutomatedPerfTestEndEvent);
