APT: Update Locked Dynamic Res Screen Percentage

* Reducing value to ensure it does not exceed the max screen % for all platforms.

#rnx

diff --git a/Engine/Plugins/Performance/AutomatedPerfTesting/Config/Engine.ini b/Engine/Plugins/Performance/AutomatedPerfTesting/Config/Engine.ini
--- a/Engine/Plugins/Performance/AutomatedPerfTesting/Config/Engine.ini
+++ b/Engine/Plugins/Performance/AutomatedPerfTesting/Config/Engine.ini
@@ -1,3 +1,3 @@
 [SystemSettings]
 ; if you need to override this per platform, this needs to be done in Engine/Platforms/<Platform>/Plugins/Performance/AutomatedPerfTesting/Config/<Platform>Engine.ini
-APT.DynamicRes.LockedScreenPercentage=75
+APT.DynamicRes.LockedScreenPercentage=60
