Fix Non-Unity build error

#rnx

diff --git a/Engine/Plugins/Performance/AutomatedPerfTesting/Source/AutomatedPerfTesting/Private/AutomatedProfileGoTest.cpp b/Engine/Plugins/Performance/AutomatedPerfTesting/Source/AutomatedPerfTesting/Private/AutomatedProfileGoTest.cpp
--- a/Engine/Plugins/Performance/AutomatedPerfTesting/Source/AutomatedPerfTesting/Private/AutomatedProfileGoTest.cpp
+++ b/Engine/Plugins/Performance/AutomatedPerfTesting/Source/AutomatedPerfTesting/Private/AutomatedProfileGoTest.cpp
@@ -9,6 +9,7 @@
 
 #include "GameFramework/PlayerController.h"
 #include "Misc/CommandLine.h"
+#include "Misc/Paths.h"
 #include "JsonObjectConverter.h"
 #include "Logging/StructuredLog.h"
 #include "UObject/UObjectIterator.h"
diff --git a/Engine/Plugins/Performance/AutomatedPerfTesting/Source/AutomatedPerfTesting/Private/ProfileGo/ProfileGo.cpp b/Engine/Plugins/Performance/AutomatedPerfTesting/Source/AutomatedPerfTesting/Private/ProfileGo/ProfileGo.cpp
--- a/Engine/Plugins/Performance/AutomatedPerfTesting/Source/AutomatedPerfTesting/Private/ProfileGo/ProfileGo.cpp
+++ b/Engine/Plugins/Performance/AutomatedPerfTesting/Source/AutomatedPerfTesting/Private/ProfileGo/ProfileGo.cpp
@@ -6,6 +6,7 @@
 #include "AutomatedProfileGoTest.h"
 #include "AutomatedPerfTesting.h"
 #include "AutomatedPerfTestControllerBase.h"
+#include "Algo/Find.h"
 
 #include "CoreMinimal.h"
 #include "ContentStreaming.h"
