Fix the APT script generator to produce scripts that work regardless of the current working directory.

The templates previously contained relative paths. The generated scripts only succeed when run from the directory containing the script.

diff --git a/Engine/Plugins/Performance/AutomatedPerfTesting/Resources/TemplatesForHookup/RunLocalTests.bat b/Engine/Plugins/Performance/AutomatedPerfTesting/Resources/TemplatesForHookup/RunLocalTests.bat
--- a/Engine/Plugins/Performance/AutomatedPerfTesting/Resources/TemplatesForHookup/RunLocalTests.bat
+++ b/Engine/Plugins/Performance/AutomatedPerfTesting/Resources/TemplatesForHookup/RunLocalTests.bat
@@ -1,3 +1,3 @@
-call **RUN_UAT_RELATIVEPATH**\RunUAT.bat BuildGraph -Script=**REPLACE_PROJECTPATH**/Build/AutoPerfTests.xml -Target="BuildAndTest **REPLACE_PROJECTNAME**" -UseLocalBuildStorage -set:TargetPlatforms=Win64 -set:TargetConfigurations=Test -set:TargetBootTest=false -set:WithAPTSequenceTests=true -set:DoPerf=true -set:DoLLM=true -set:GenerateLocalReports=true -set:IterationsPerf=1 -set:IterationsLLM=1 -set:IterationsInsights=1
+call "%~dp0**RUN_UAT_RELATIVEPATH**\RunUAT.bat" BuildGraph -Script="**REPLACE_PROJECTPATH**/Build/AutoPerfTests.xml" -Target="BuildAndTest **REPLACE_PROJECTNAME**" -UseLocalBuildStorage -set:TargetPlatforms=Win64 -set:TargetConfigurations=Test -set:TargetBootTest=false -set:WithAPTSequenceTests=true -set:DoPerf=true -set:DoLLM=true -set:GenerateLocalReports=true -set:IterationsPerf=1 -set:IterationsLLM=1 -set:IterationsInsights=1
 pause
 
diff --git a/Engine/Plugins/Performance/AutomatedPerfTesting/Resources/TemplatesForHookup/RunLocalTests.sh b/Engine/Plugins/Performance/AutomatedPerfTesting/Resources/TemplatesForHookup/RunLocalTests.sh
--- a/Engine/Plugins/Performance/AutomatedPerfTesting/Resources/TemplatesForHookup/RunLocalTests.sh
+++ b/Engine/Plugins/Performance/AutomatedPerfTesting/Resources/TemplatesForHookup/RunLocalTests.sh
@@ -1,3 +1,7 @@
+#!/bin/bash
+
+SCRIPT_DIR=$( cd -- "$( dirname -- "${BASH_SOURCE[0]}" )" &> /dev/null && pwd )
+
 platform=""
 if [ "$(uname)" = "Darwin" ]; then
 	platform="Mac"
@@ -7,4 +11,4 @@ if [ "$(uname)" = "Linux" ]; then
 	platform="Linux"
 fi
 
-**RUN_UAT_RELATIVEPATH**/RunUAT.sh BuildGraph -Script=**REPLACE_PROJECTPATH**/Build/AutoPerfTests.xml -Target="BuildAndTest **REPLACE_PROJECTNAME**" -UseLocalBuildStorage -set:TargetPlatforms=$platform -set:WithAPTSequenceTests=true -set:TargetConfigurations=Test -set:DoPerf=true -set:DoLLM=true -set:GenerateLocalReports=true -set:GenerateLocalReports=true -set:IterationsPerf=1 -set:IterationsLLM=1 -set:IterationsInsights=1
+"$SCRIPT_DIR"/**RUN_UAT_RELATIVEPATH**/RunUAT.sh BuildGraph -Script=**REPLACE_PROJECTPATH**/Build/AutoPerfTests.xml -Target="BuildAndTest **REPLACE_PROJECTNAME**" -UseLocalBuildStorage -set:TargetPlatforms=$platform -set:WithAPTSequenceTests=true -set:TargetConfigurations=Test -set:DoPerf=true -set:DoLLM=true -set:GenerateLocalReports=true -set:GenerateLocalReports=true -set:IterationsPerf=1 -set:IterationsLLM=1 -set:IterationsInsights=1
