Fix APT/PL2 dependency issue

#rnx

diff --git a/Engine/Plugins/Performance/AutomatedPerfTesting/AutomatedPerfTesting.uplugin b/Engine/Plugins/Performance/AutomatedPerfTesting/AutomatedPerfTesting.uplugin
--- a/Engine/Plugins/Performance/AutomatedPerfTesting/AutomatedPerfTesting.uplugin
+++ b/Engine/Plugins/Performance/AutomatedPerfTesting/AutomatedPerfTesting.uplugin
@@ -28,7 +28,12 @@
 		},
 		{
 			"Name": "ProjectLauncher",
-			"Enabled":  true
+			"Enabled": true,
+			"SupportedTargetPlatforms": [
+				"Win64",
+				"Linux",
+				"Mac"
+			]
 		}
 	]
 }
\ No newline at end of file
diff --git a/Engine/Plugins/Performance/AutomatedPerfTesting/Source/AutomatedPerfTesting/AutomatedPerfTesting.Build.cs b/Engine/Plugins/Performance/AutomatedPerfTesting/Source/AutomatedPerfTesting/AutomatedPerfTesting.Build.cs
--- a/Engine/Plugins/Performance/AutomatedPerfTesting/Source/AutomatedPerfTesting/AutomatedPerfTesting.Build.cs
+++ b/Engine/Plugins/Performance/AutomatedPerfTesting/Source/AutomatedPerfTesting/AutomatedPerfTesting.Build.cs
@@ -1,9 +1,16 @@
 // Copyright Epic Games, Inc. All Rights Reserved.
 
+using System.Linq;
 using UnrealBuildTool;
 
 public class AutomatedPerfTesting : ModuleRules
 {
+	UnrealTargetPlatform[] SupportedProjectLauncherPlatforms => [
+		UnrealTargetPlatform.Win64,
+		UnrealTargetPlatform.Linux,
+		UnrealTargetPlatform.Mac
+	];
+
 	public AutomatedPerfTesting(ReadOnlyTargetRules Target) : base(Target)
 	{
 		PCHUsage = ModuleRules.PCHUsageMode.UseExplicitOrSharedPCHs;
@@ -30,7 +37,7 @@ public class AutomatedPerfTesting : ModuleRules
 			}
 			);
 
-		if(Target.bCompileAgainstEditor)
+		if(Target.bCompileAgainstEditor && SupportedProjectLauncherPlatforms.Contains(Target.Platform))
 		{
 			PublicDependencyModuleNames.AddRange(
 			[
