[APT]
Fix crash in APT due to MultiWorld having multiple valid Game Worlds alive at one time
- The assumption was we would grab the first available Game World and see if it is the one we were looking for
- We instead loop over the list of Game Worlds and check the names of them to see if they match the target World
- Note: this whole bit of logic will need to be revisited as we refactor APT

#rnx

diff --git a/Engine/Plugins/Performance/AutomatedPerfTesting/Source/AutomatedPerfTesting/Private/AutomatedSequencePerfTest.cpp b/Engine/Plugins/Performance/AutomatedPerfTesting/Source/AutomatedPerfTesting/Private/AutomatedSequencePerfTest.cpp
--- a/Engine/Plugins/Performance/AutomatedPerfTesting/Source/AutomatedPerfTesting/Private/AutomatedSequencePerfTest.cpp
+++ b/Engine/Plugins/Performance/AutomatedPerfTesting/Source/AutomatedPerfTesting/Private/AutomatedSequencePerfTest.cpp
@@ -131,9 +131,10 @@ void UAutomatedSequencePerfTest::NextMap()
 			OptionsString += "game=" + CurrentMapSequenceCombo->GameModeOverride;
 		}
 		
-		if (OptionsString.IsEmpty() && AutomatedPerfTest::FindCurrentWorld()->GetName() == CurrentMapSequenceCombo->Map.GetAssetName())
+		if (OptionsString.IsEmpty() && AutomatedPerfTest::IsWorldLoaded(CurrentMapSequenceCombo->Map.GetAssetName()))
 		{
 			UE_LOG(LogAutomatedPerfTest, Log, TEXT("%s is already loaded and does not have any options string, skipping map load and setting up the test"), *CurrentMapSequenceCombo->Map.GetAssetName());
+			// TODO: Fix up SetupTest's logic to also correctly deal with MultiWorld
 			SetupTest();
 		}
 		else 
diff --git a/Engine/Plugins/Performance/AutomatedPerfTesting/Source/AutomatedPerfTesting/Public/AutomatedPerfTestControllerBase.h b/Engine/Plugins/Performance/AutomatedPerfTesting/Source/AutomatedPerfTesting/Public/AutomatedPerfTestControllerBase.h
--- a/Engine/Plugins/Performance/AutomatedPerfTesting/Source/AutomatedPerfTesting/Public/AutomatedPerfTestControllerBase.h
+++ b/Engine/Plugins/Performance/AutomatedPerfTesting/Source/AutomatedPerfTesting/Public/AutomatedPerfTestControllerBase.h
@@ -50,6 +50,34 @@ namespace AutomatedPerfTest
 
 		return World;
 	}
+
+	// New function because there could be more than one Game world alive at any point with MultiWorld
+	static bool IsWorldLoaded(const FString & WorldName)
+	{
+		for (const FWorldContext& WorldContext : GEngine->GetWorldContexts())
+		{
+			if (WorldContext.WorldType == EWorldType::Game)
+			{
+				UWorld* World = WorldContext.World();
+				if (World && World->GetName().Compare(WorldName) == 0)
+				{
+					return true;
+				}
+			}
+#if WITH_EDITOR
+			else if (GIsEditor && WorldContext.WorldType == EWorldType::PIE)
+			{
+				UWorld* World = WorldContext.World();
+				if (World && World->GetName().Compare(WorldName) == 0)
+				{
+					return true;
+				}
+			}
+#endif
+		}
+
+		return false;
+	}
 }
 /**
  * 
