Removing unused OnGameStateSet delegate from UAutomatedPerfTestControllerBase

OnGameStateSet was registered as a listener for UWorld::GameStateSetEvent. Its handler only logged "Game State Set" and then removed the delegate. None of the subclasses relied on this base-class behavior, and any subclass that needs this functionality can implement its own handler, so the delegate and associated code have been removed.

#Synced-CL 48813241

diff --git a/Engine/Plugins/Performance/AutomatedPerfTesting/Source/AutomatedPerfTesting/Private/AutomatedPerfTestControllerBase.cpp b/Engine/Plugins/Performance/AutomatedPerfTesting/Source/AutomatedPerfTesting/Private/AutomatedPerfTestControllerBase.cpp
--- a/Engine/Plugins/Performance/AutomatedPerfTesting/Source/AutomatedPerfTesting/Private/AutomatedPerfTestControllerBase.cpp
+++ b/Engine/Plugins/Performance/AutomatedPerfTesting/Source/AutomatedPerfTesting/Private/AutomatedPerfTestControllerBase.cpp
@@ -58,7 +58,6 @@ void UAutomatedPerfTestControllerBase::OnPreWorldInitializeInternal(UWorld* Worl
 void UAutomatedPerfTestControllerBase::OnPreWorldInitialize(UWorld* World)
 {
 	check(World);
-	World->GameStateSetEvent.AddUObject(this, &ThisClass::OnGameStateSet);
 	World->OnWorldBeginPlay.AddUObject(this, &ThisClass::OnWorldBeginPlay);
 }
 
@@ -102,15 +101,6 @@ void UAutomatedPerfTestControllerBase::OnWorldBeginPlay()
 	SetupTest();
 }
 
-void UAutomatedPerfTestControllerBase::OnGameStateSet(AGameStateBase* const GameStateBase)
-{
-	UE_LOG(LogAutomatedPerfTest, Log, TEXT("Game State Set"));
-	if (UWorld* const World = GetWorld())
-	{
-		World->GameStateSetEvent.RemoveAll(this);
-	}
-}
-
 UAutomatedPerfTestControllerBase::UAutomatedPerfTestControllerBase(const FObjectInitializer& ObjectInitializer)
 	: Super(ObjectInitializer)
 	, TraceChannels("default,screenshot,stats")
diff --git a/Engine/Plugins/Performance/AutomatedPerfTesting/Source/AutomatedPerfTesting/Public/AutomatedPerfTestControllerBase.h b/Engine/Plugins/Performance/AutomatedPerfTesting/Source/AutomatedPerfTesting/Public/AutomatedPerfTestControllerBase.h
--- a/Engine/Plugins/Performance/AutomatedPerfTesting/Source/AutomatedPerfTesting/Public/AutomatedPerfTestControllerBase.h
+++ b/Engine/Plugins/Performance/AutomatedPerfTesting/Source/AutomatedPerfTesting/Public/AutomatedPerfTestControllerBase.h
@@ -96,9 +96,6 @@ public:
 	UFUNCTION()
 	UE_API void OnWorldBeginPlay();
 
-	UFUNCTION()
-	UE_API void OnGameStateSet(AGameStateBase* const GameStateBase);
-
 // Base functionality
 public:
 	UE_API UAutomatedPerfTestControllerBase(const FObjectInitializer& ObjectInitializer);
