Removing the unused UAutomatedPerfTestControllerBase::TryEarlyExec

This function was a workaround which is no longer needed due to improvements elsewhere in Device Profile and TestConfigs.

#Synced-CL 48813241

diff --git a/Engine/Plugins/Performance/AutomatedPerfTesting/Source/AutomatedPerfTesting/Private/AutomatedPerfTestControllerBase.cpp b/Engine/Plugins/Performance/AutomatedPerfTesting/Source/AutomatedPerfTesting/Private/AutomatedPerfTestControllerBase.cpp
--- a/Engine/Plugins/Performance/AutomatedPerfTesting/Source/AutomatedPerfTesting/Private/AutomatedPerfTestControllerBase.cpp
+++ b/Engine/Plugins/Performance/AutomatedPerfTesting/Source/AutomatedPerfTesting/Private/AutomatedPerfTestControllerBase.cpp
@@ -38,7 +38,6 @@ static FAutoConsoleVariableRef CVarDynamicResLockedScreenPercentage(
 
 void UAutomatedPerfTestControllerBase::OnPreWorldInitializeInternal(UWorld* World, const UWorld::InitializationValues IVS)
 {
-	TryEarlyExec(World);
 	OnPreWorldInitialize(World);
 
 	if (RequestsLockedDynRes())
@@ -61,40 +60,6 @@ void UAutomatedPerfTestControllerBase::OnPreWorldInitialize(UWorld* World)
 	World->OnWorldBeginPlay.AddUObject(this, &ThisClass::OnWorldBeginPlay);
 }
 
-void UAutomatedPerfTestControllerBase::TryEarlyExec(UWorld* const World)
-{
-	check(World);
-	
-	if (GEngine)
-	{
-		// Search the list of deferred commands
-		const TArray<FString>& DeferredCmds = GEngine->DeferredCommands;
-		TArray<int32> ExecutedIndices;
-		for (int32 DeferredCmdIndex = 0; DeferredCmdIndex < DeferredCmds.Num(); ++DeferredCmdIndex)
-		{
-			// If the deferred command is one that should be executed early
-			const FString& DeferredCmd = DeferredCmds[DeferredCmdIndex];
-			if (CmdsToExecEarly.ContainsByPredicate([&DeferredCmd](const FString& CmdToFind) { return DeferredCmd.StartsWith(CmdToFind); }))
-			{
-				UE_LOG(LogAutomatedPerfTest, Log, TEXT("EarlyExec: executing '%s' early."), *DeferredCmd);
-				GEngine->Exec(World, *DeferredCmd);
-				ExecutedIndices.Push(DeferredCmdIndex);
-			}
-		}
-
-		// Remove the executed commands from the list of deferred commands
-		// Note: This is done in reverse order to ensure the cached indices remain valid
-		while (!ExecutedIndices.IsEmpty())
-		{
-			GEngine->DeferredCommands.RemoveAt(ExecutedIndices.Pop());
-		}
-	}
-	else
-	{
-		UE_LOG(LogAutomatedPerfTest, Warning, TEXT("Attempted EarlyExec without GEngine being ready"))
-	}
-}
-
 void UAutomatedPerfTestControllerBase::OnWorldBeginPlay()
 {
 	UE_LOG(LogAutomatedPerfTest, Log, TEXT("OnWorldBeginPlay"));
diff --git a/Engine/Plugins/Performance/AutomatedPerfTesting/Source/AutomatedPerfTesting/Public/AutomatedPerfTestControllerBase.h b/Engine/Plugins/Performance/AutomatedPerfTesting/Source/AutomatedPerfTesting/Public/AutomatedPerfTestControllerBase.h
--- a/Engine/Plugins/Performance/AutomatedPerfTesting/Source/AutomatedPerfTesting/Public/AutomatedPerfTestControllerBase.h
+++ b/Engine/Plugins/Performance/AutomatedPerfTesting/Source/AutomatedPerfTesting/Public/AutomatedPerfTestControllerBase.h
@@ -90,9 +90,6 @@ public:
 	UE_API void OnPreWorldInitializeInternal(UWorld* World, const UWorld::InitializationValues IVS);
 	UE_API virtual void OnPreWorldInitialize(UWorld* World);
 
-	UFUNCTION()
-	UE_API void TryEarlyExec(UWorld* const World);
-	
 	UFUNCTION()
 	UE_API void OnWorldBeginPlay();
 
@@ -192,8 +189,6 @@ private:
 
 	FText VideoRecordingTitle;
 	
-	const TArray<FString> CmdsToExecEarly = { };
-	
 	AGameModeBase* GameMode;
 
 	FDelegateHandle CsvProfilerDelegateHandle;
