Removing unimplemented code in the Automated Performance Testing plugin.

- Removed virtual function overrides that were only forwarding to the base class and did not add behavior or specialization.
- Removed the AutomatedSoakTest cpp/h files which are stubs with no implementation. These can be reintroduced  in the future if needed.

#Synced-CL 49635284

diff --git a/Engine/Plugins/Performance/AutomatedPerfTesting/Source/AutomatedPerfTesting/Private/AutomatedPerfTestControllerBase.cpp b/Engine/Plugins/Performance/AutomatedPerfTesting/Source/AutomatedPerfTesting/Private/AutomatedPerfTestControllerBase.cpp
--- a/Engine/Plugins/Performance/AutomatedPerfTesting/Source/AutomatedPerfTesting/Private/AutomatedPerfTestControllerBase.cpp
+++ b/Engine/Plugins/Performance/AutomatedPerfTesting/Source/AutomatedPerfTesting/Private/AutomatedPerfTestControllerBase.cpp
@@ -607,16 +607,6 @@ void UAutomatedPerfTestControllerBase::OnTick(float TimeDelta)
 	MarkHeartbeatActive();
 }
 
-void UAutomatedPerfTestControllerBase::OnStateChange(FName OldState, FName NewState)
-{
-	Super::OnStateChange(OldState, NewState);
-}
-
-void UAutomatedPerfTestControllerBase::OnPreMapChange()
-{
-	Super::OnPreMapChange();
-}
-
 void UAutomatedPerfTestControllerBase::BeginDestroy()
 {
 	UnbindAllDelegates();
diff --git a/Engine/Plugins/Performance/AutomatedPerfTesting/Source/AutomatedPerfTesting/Private/AutomatedProfileGoTest.cpp b/Engine/Plugins/Performance/AutomatedPerfTesting/Source/AutomatedPerfTesting/Private/AutomatedProfileGoTest.cpp
--- a/Engine/Plugins/Performance/AutomatedPerfTesting/Source/AutomatedPerfTesting/Private/AutomatedProfileGoTest.cpp
+++ b/Engine/Plugins/Performance/AutomatedPerfTesting/Source/AutomatedPerfTesting/Private/AutomatedProfileGoTest.cpp
@@ -220,21 +220,6 @@ void UAutomatedProfileGoTest::OnTick(float TimeDelta)
 	Super::OnTick(TimeDelta);
 }
 
-void UAutomatedProfileGoTest::RunTest()
-{
-	Super::RunTest();
-}
-
-void UAutomatedProfileGoTest::TeardownTest(bool bExitAfterTeardown)
-{
-	Super::TeardownTest(bExitAfterTeardown);
-}
-
-void UAutomatedProfileGoTest::Exit()
-{
-	Super::Exit();
-}
-
 void UAutomatedProfileGoTest::OnRequestFailed()
 {
 	UE_LOGFMT(LogAutomatedPerfTest, Error, "ProfileGo request failed. Exiting.");
diff --git a/Engine/Plugins/Performance/AutomatedPerfTesting/Source/AutomatedPerfTesting/Private/AutomatedReplayPerfTest.cpp b/Engine/Plugins/Performance/AutomatedPerfTesting/Source/AutomatedPerfTesting/Private/AutomatedReplayPerfTest.cpp
--- a/Engine/Plugins/Performance/AutomatedPerfTesting/Source/AutomatedPerfTesting/Private/AutomatedReplayPerfTest.cpp
+++ b/Engine/Plugins/Performance/AutomatedPerfTesting/Source/AutomatedPerfTesting/Private/AutomatedReplayPerfTest.cpp
@@ -123,11 +123,6 @@ void UAutomatedReplayPerfTest::TeardownTest(bool bExitAfterTeardown)
 	Super::TeardownTest(bExitAfterTeardown);
 }
 
-void UAutomatedReplayPerfTest::Exit()
-{
-	Super::Exit();
-}
-
 void UAutomatedReplayPerfTest::OnInit()
 {
 	Super::OnInit();
diff --git a/Engine/Plugins/Performance/AutomatedPerfTesting/Source/AutomatedPerfTesting/Private/AutomatedSequencePerfTest.cpp b/Engine/Plugins/Performance/AutomatedPerfTesting/Source/AutomatedPerfTesting/Private/AutomatedSequencePerfTest.cpp
--- a/Engine/Plugins/Performance/AutomatedPerfTesting/Source/AutomatedPerfTesting/Private/AutomatedSequencePerfTest.cpp
+++ b/Engine/Plugins/Performance/AutomatedPerfTesting/Source/AutomatedPerfTesting/Private/AutomatedSequencePerfTest.cpp
@@ -206,12 +206,6 @@ void UAutomatedSequencePerfTest::TeardownTest(bool bExitAfterTeardown)
 	NextMap();
 }
 
-void UAutomatedSequencePerfTest::Exit()
-{
-	UE_LOG(LogAutomatedPerfTest, Log, TEXT("AutomatedSequencePerfTest::Exit"));
-	Super::Exit();
-}
-
 bool UAutomatedSequencePerfTest::TryStartCSVProfiler(FString CSVFileName, const FString& DestinationFolder, int32 Frames)
 {
 #if CSV_PROFILER
diff --git a/Engine/Plugins/Performance/AutomatedPerfTesting/Source/AutomatedPerfTesting/Private/AutomatedSoakTest.cpp b/Engine/Plugins/Performance/AutomatedPerfTesting/Source/AutomatedPerfTesting/Private/AutomatedSoakTest.cpp
deleted file mode 100644
diff --git a/Engine/Plugins/Performance/AutomatedPerfTesting/Source/AutomatedPerfTesting/Public/AutomatedPerfTestControllerBase.h b/Engine/Plugins/Performance/AutomatedPerfTesting/Source/AutomatedPerfTesting/Public/AutomatedPerfTestControllerBase.h
--- a/Engine/Plugins/Performance/AutomatedPerfTesting/Source/AutomatedPerfTesting/Public/AutomatedPerfTestControllerBase.h
+++ b/Engine/Plugins/Performance/AutomatedPerfTesting/Source/AutomatedPerfTesting/Public/AutomatedPerfTestControllerBase.h
@@ -145,8 +145,6 @@ protected:
 	// ~Begin UGauntletTestController Interface
 	UE_API virtual void OnInit() override;
 	UE_API virtual void OnTick(float TimeDelta) override;
-	UE_API virtual void OnStateChange(FName OldState, FName NewState) override;
-	UE_API virtual void OnPreMapChange() override;
 	UE_API virtual void BeginDestroy() override;
 	// ~End UGauntletTestController Interface
 
diff --git a/Engine/Plugins/Performance/AutomatedPerfTesting/Source/AutomatedPerfTesting/Public/AutomatedProfileGoTest.h b/Engine/Plugins/Performance/AutomatedPerfTesting/Source/AutomatedPerfTesting/Public/AutomatedProfileGoTest.h
--- a/Engine/Plugins/Performance/AutomatedPerfTesting/Source/AutomatedPerfTesting/Public/AutomatedProfileGoTest.h
+++ b/Engine/Plugins/Performance/AutomatedPerfTesting/Source/AutomatedPerfTesting/Public/AutomatedProfileGoTest.h
@@ -35,11 +35,6 @@ public:
 	UE_API virtual void SetupTest() override;
 	UE_API virtual void OnTick(float TimeDelta) override;
 
-	UFUNCTION()
-	UE_API virtual void RunTest() override;
-	UE_API virtual void TeardownTest(bool bExitAfterTeardown = true) override;
-	UE_API virtual void Exit() override;
-
 protected:
 	UE_API virtual void OnRequestFailed();
 	UE_API virtual void OnScenarioStarted(const FString& InName);
diff --git a/Engine/Plugins/Performance/AutomatedPerfTesting/Source/AutomatedPerfTesting/Public/AutomatedReplayPerfTest.h b/Engine/Plugins/Performance/AutomatedPerfTesting/Source/AutomatedPerfTesting/Public/AutomatedReplayPerfTest.h
--- a/Engine/Plugins/Performance/AutomatedPerfTesting/Source/AutomatedPerfTesting/Public/AutomatedReplayPerfTest.h
+++ b/Engine/Plugins/Performance/AutomatedPerfTesting/Source/AutomatedPerfTesting/Public/AutomatedReplayPerfTest.h
@@ -57,7 +57,6 @@ public:
 	UFUNCTION()
 	UE_API virtual void RunTest() override;
 	UE_API virtual void TeardownTest(bool bExitAfterTeardown = true) override;
-	UE_API virtual void Exit() override;
 
 protected:
 	UE_API virtual void OnInit() override;
diff --git a/Engine/Plugins/Performance/AutomatedPerfTesting/Source/AutomatedPerfTesting/Public/AutomatedSequencePerfTest.h b/Engine/Plugins/Performance/AutomatedPerfTesting/Source/AutomatedPerfTesting/Public/AutomatedSequencePerfTest.h
--- a/Engine/Plugins/Performance/AutomatedPerfTesting/Source/AutomatedPerfTesting/Public/AutomatedSequencePerfTest.h
+++ b/Engine/Plugins/Performance/AutomatedPerfTesting/Source/AutomatedPerfTesting/Public/AutomatedSequencePerfTest.h
@@ -97,7 +97,6 @@ public:
 	UE_API virtual void RunTest() override;
 	
 	UE_API virtual void TeardownTest(bool bExitAfterTeardown = true) override;
-	UE_API virtual void Exit() override;
 
 	UE_API virtual bool TryStartCSVProfiler(FString CSVFileName, const FString& DestinationFolder = FString(), int32 Frames = -1) override;
 	// ~End UAutomatedPerfTestControllerBase Interface
diff --git a/Engine/Plugins/Performance/AutomatedPerfTesting/Source/AutomatedPerfTesting/Public/AutomatedSoakTest.h b/Engine/Plugins/Performance/AutomatedPerfTesting/Source/AutomatedPerfTesting/Public/AutomatedSoakTest.h
deleted file mode 100644
