Remove unused member variables in Automated Performance Testing Gauntlet controller classes.

#Synced-CL 49664943

diff --git a/Engine/Plugins/Performance/AutomatedPerfTesting/Source/AutomatedPerfTesting/Private/AutomatedProfileGoTest.cpp b/Engine/Plugins/Performance/AutomatedPerfTesting/Source/AutomatedPerfTesting/Private/AutomatedProfileGoTest.cpp
--- a/Engine/Plugins/Performance/AutomatedPerfTesting/Source/AutomatedPerfTesting/Private/AutomatedProfileGoTest.cpp
+++ b/Engine/Plugins/Performance/AutomatedPerfTesting/Source/AutomatedPerfTesting/Private/AutomatedProfileGoTest.cpp
@@ -57,7 +57,6 @@ UAutomatedProfileGoTest::UAutomatedProfileGoTest(const FObjectInitializer& Objec
 	bRequestedProfileGo = false;
 	bStartedProfileGo = false;
 	bShouldTick = true;
-	TimeInCorrectState = 0.0f;
 	NumRequiredLoops = 0;
 	NumCompletedLoops = 0;
 }
diff --git a/Engine/Plugins/Performance/AutomatedPerfTesting/Source/AutomatedPerfTesting/Private/LaunchExtension/AutomatedPerfTestLaunchExtension.h b/Engine/Plugins/Performance/AutomatedPerfTesting/Source/AutomatedPerfTesting/Private/LaunchExtension/AutomatedPerfTestLaunchExtension.h
--- a/Engine/Plugins/Performance/AutomatedPerfTesting/Source/AutomatedPerfTesting/Private/LaunchExtension/AutomatedPerfTestLaunchExtension.h
+++ b/Engine/Plugins/Performance/AutomatedPerfTesting/Source/AutomatedPerfTesting/Private/LaunchExtension/AutomatedPerfTestLaunchExtension.h
@@ -65,7 +65,6 @@ private:
 	}
 
 	EAutomatedPerfTestType CurrentTestType = EAutomatedPerfTestType::Sequence;
-	bool bEnableCSVProfiler = false;
 };
 
 
diff --git a/Engine/Plugins/Performance/AutomatedPerfTesting/Source/AutomatedPerfTesting/Public/AutomatedMaterialPerfTest.h b/Engine/Plugins/Performance/AutomatedPerfTesting/Source/AutomatedPerfTesting/Public/AutomatedMaterialPerfTest.h
--- a/Engine/Plugins/Performance/AutomatedPerfTesting/Source/AutomatedPerfTesting/Public/AutomatedMaterialPerfTest.h
+++ b/Engine/Plugins/Performance/AutomatedPerfTesting/Source/AutomatedPerfTesting/Public/AutomatedMaterialPerfTest.h
@@ -147,7 +147,6 @@ protected:
 private:
 	ACameraActor* Camera;
 	AStaticMeshActor* MaterialPlate;
-	FString CurrentRegionName;
 	UMaterialInterface* CurrentMaterial;
 	int CurrentMaterialIndex;
 	const UAutomatedMaterialPerfTestProjectSettings* Settings;
diff --git a/Engine/Plugins/Performance/AutomatedPerfTesting/Source/AutomatedPerfTesting/Public/AutomatedProfileGoTest.h b/Engine/Plugins/Performance/AutomatedPerfTesting/Source/AutomatedPerfTesting/Public/AutomatedProfileGoTest.h
--- a/Engine/Plugins/Performance/AutomatedPerfTesting/Source/AutomatedPerfTesting/Public/AutomatedProfileGoTest.h
+++ b/Engine/Plugins/Performance/AutomatedPerfTesting/Source/AutomatedPerfTesting/Public/AutomatedProfileGoTest.h
@@ -54,7 +54,6 @@ protected:
 	bool	bRequestedProfileGo;
 	bool	bStartedProfileGo;
 	bool	bShouldTick;
-	float	TimeInCorrectState;
 	FString ProfileGoStatus;
 	FString ProfileGoConfigFile;
 
diff --git a/Engine/Plugins/Performance/AutomatedPerfTesting/Source/AutomatedPerfTesting/Public/StaticCameraTests/AutomatedStaticCameraPerfTestBase.h b/Engine/Plugins/Performance/AutomatedPerfTesting/Source/AutomatedPerfTesting/Public/StaticCameraTests/AutomatedStaticCameraPerfTestBase.h
--- a/Engine/Plugins/Performance/AutomatedPerfTesting/Source/AutomatedPerfTesting/Public/StaticCameraTests/AutomatedStaticCameraPerfTestBase.h
+++ b/Engine/Plugins/Performance/AutomatedPerfTesting/Source/AutomatedPerfTesting/Public/StaticCameraTests/AutomatedStaticCameraPerfTestBase.h
@@ -13,8 +13,6 @@ UCLASS(MinimalAPI, BlueprintType, Config=Engine, DefaultConfig, DisplayName="Aut
 class UAutomatedStaticCameraPerfTestProjectSettings : public UDeveloperSettings
 {
 	GENERATED_BODY()
-
-	TMap<FString, FSoftObjectPath> MapNameMap; 
 	
 public:
 	UE_API UAutomatedStaticCameraPerfTestProjectSettings(const FObjectInitializer&);
