Removed Automated Performance Testing controller support for specifying CSVOutputMode.

UAutomatedPerfTestControllerBase subclasses previously exposed a configuration option to control whether a performance test emitted multiple CSVs or a single CSV with CSV_EVENTS used to annotate meaningful phases within tests. In practice, the default value per test type was the configuration that made the most sense for data review and downstream aggregation.

This change removes the configuration and codifies the effective behavior:

Materials - One CSV per material tested (previously 'CSVOutputMode=Granular')
StaticCamera - One CSV per performance camera (previously 'CSVOutputMode=Granular')
Replay - One CSV per replay (previously 'CSVOutputMode=Separate')
Sequence - One CSV per sequence (previously 'CSVOutputMode=Separate')

As a result, the CSV_EVENTS used to delineate these phases in single-CSV mode are no longer required and have been removed.

#Synced-CL 49664943

diff --git a/Engine/Plugins/Performance/AutomatedPerfTesting/Source/AutomatedPerfTesting/Private/AutomatedMaterialPerfTest.cpp b/Engine/Plugins/Performance/AutomatedPerfTesting/Source/AutomatedPerfTesting/Private/AutomatedMaterialPerfTest.cpp
--- a/Engine/Plugins/Performance/AutomatedPerfTesting/Source/AutomatedPerfTesting/Private/AutomatedMaterialPerfTest.cpp
+++ b/Engine/Plugins/Performance/AutomatedPerfTesting/Source/AutomatedPerfTesting/Private/AutomatedMaterialPerfTest.cpp
@@ -25,7 +25,6 @@ UAutomatedMaterialPerfTestProjectSettings::UAutomatedMaterialPerfTestProjectSett
 	, CameraProjectionMode(ECameraProjectionMode::Type::Orthographic)
 	, PlateDistanceFromCamera(512.0)
 	, MaterialPlate(FSoftObjectPath("/AutomatedPerfTesting/Tests/Materials/SM_AutomatedMaterialPerfTestDefaultPlate.SM_AutomatedMaterialPerfTestDefaultPlate"))
-	, CSVOutputMode(EAutomatedPerfTestCSVOutputMode::Granular)
 {
 }
 
@@ -112,10 +111,7 @@ bool UAutomatedMaterialPerfTest::TryStartCSVProfiler(FString CSVFileName, const
 #if CSV_PROFILER
 	if(FCsvProfiler* const CsvProfiler = FCsvProfiler::Get())
 	{
-		if (GetCSVOutputMode() == EAutomatedPerfTestCSVOutputMode::Granular)
-		{
-			CsvProfiler->SetMetadata(TEXT("Material"), *GetCurrentMaterialRegionName());
-		}
+		CsvProfiler->SetMetadata(TEXT("Material"), *GetCurrentMaterialRegionName());
 		return Super::TryStartCSVProfiler(CSVFileName, DestinationFolder, Frames);
 	}
 #endif
@@ -209,11 +205,7 @@ void UAutomatedMaterialPerfTest::MarkMaterialStart()
 #if CSV_PROFILER
 		if(RequestsCSVProfiler())
 		{
-			if(GetCSVOutputMode() == EAutomatedPerfTestCSVOutputMode::Granular)
-			{
-				TryStartCSVProfiler(GetCurrentMaterialRegionFullName());
-			}
-			CSV_EVENT(AutomatedPerfTest, TEXT("START_%s"), *GetCurrentMaterialRegionName())
+			TryStartCSVProfiler(GetCurrentMaterialRegionFullName());
 		}
 #endif
 	}
@@ -231,11 +223,7 @@ void UAutomatedMaterialPerfTest::MarkMaterialEnd()
 #if CSV_PROFILER
 		if(RequestsCSVProfiler())
 		{
-			CSV_EVENT(AutomatedPerfTest, TEXT("END_%s"), *GetCurrentMaterialRegionName())
-			if(GetCSVOutputMode() == EAutomatedPerfTestCSVOutputMode::Granular)
-			{
-				TryStopCSVProfiler();
-			}
+			TryStopCSVProfiler();
 		}
 #endif
 	}
@@ -249,8 +237,6 @@ void UAutomatedMaterialPerfTest::OnInit()
 
 	Settings = GetDefault<UAutomatedMaterialPerfTestProjectSettings>();
 
-	SetCSVOutputMode(Settings->CSVOutputMode);
-	
 	// early out if there aren't actually any materials
 	if(Settings->MaterialsToTest.Num() < 0)
 	{
diff --git a/Engine/Plugins/Performance/AutomatedPerfTesting/Source/AutomatedPerfTesting/Private/AutomatedPerfTestControllerBase.cpp b/Engine/Plugins/Performance/AutomatedPerfTesting/Source/AutomatedPerfTesting/Private/AutomatedPerfTestControllerBase.cpp
--- a/Engine/Plugins/Performance/AutomatedPerfTesting/Source/AutomatedPerfTesting/Private/AutomatedPerfTestControllerBase.cpp
+++ b/Engine/Plugins/Performance/AutomatedPerfTesting/Source/AutomatedPerfTesting/Private/AutomatedPerfTestControllerBase.cpp
@@ -76,7 +76,6 @@ UAutomatedPerfTestControllerBase::UAutomatedPerfTestControllerBase(const FObject
 	, bRequestsLockedDynRes(false)
 	, InsightsRegionID(0)
 	, ArtifactOutputPath()
-    , CSVOutputMode(EAutomatedPerfTestCSVOutputMode::Single)
 {
 	// cache this off once, so that it's consistent throughout a session
 	TestDatetime = FDateTime::Now().ToString(TEXT("%Y%m%d-%H%M%S"));
@@ -307,8 +306,6 @@ void UAutomatedPerfTestControllerBase::SetupTest()
 	SetupGameModeInstance();
 
 	// Subclasses should implement their own transitions from SetupTest to RunTest depending on their needs
-
-	UE_LOG(LogAutomatedPerfTest, Log, TEXT("CSV Output Mode: %s"), *UEnum::GetValueAsString(CSVOutputMode))
 	UE_LOG(LogGauntlet, Log, TEXT("Setup Test: %s"), *GetTestID());
 }
 
@@ -441,12 +438,6 @@ APlayerController* UAutomatedPerfTestControllerBase::GetPlayerController()
 
 void UAutomatedPerfTestControllerBase::SetupProfiling()
 {
-
-	if (RequestsCSVProfiler() && CSVOutputMode == EAutomatedPerfTestCSVOutputMode::Single)
-	{
-		TryStartCSVProfiler();
-	}
-
 	if (RequestsFPSChart())
 	{
 		TryStartFPSChart();
@@ -480,28 +471,10 @@ void UAutomatedPerfTestControllerBase::MarkProfilingStart()
 	{
 		InsightsRegionID = TRACE_BEGIN_REGION_WITH_ID(*GetOverallRegionName(), TEXT("AutomatedPerfTest"));
 	}
-
-	if (RequestsCSVProfiler())
-	{
-		if (CSVOutputMode == EAutomatedPerfTestCSVOutputMode::Separate)
-		{
-			TryStartCSVProfiler();
-		}
-		CSV_EVENT(AutomatedPerfTest, TEXT("START"), *GetOverallRegionName())
-	}
 }
 
 void UAutomatedPerfTestControllerBase::TeardownProfiling()
 {
-	if (RequestsCSVProfiler())
-	{
-		CSV_EVENT(AutomatedPerfTest, TEXT("END"), *GetOverallRegionName())
-		if (CSVOutputMode == EAutomatedPerfTestCSVOutputMode::Single)
-		{
-			TryStopCSVProfiler();
-		}
-	}
-
 	if (RequestsFPSChart())
 	{
 		TryStopFPSChart();
@@ -525,15 +498,6 @@ void UAutomatedPerfTestControllerBase::MarkProfilingEnd()
 			TRACE_END_REGION_WITH_ID(InsightsRegionID);
 		}
 	}
-
-	if (RequestsCSVProfiler())
-	{
-		CSV_EVENT(AutomatedPerfTest, TEXT("END"), *GetOverallRegionName())
-		if (CSVOutputMode == EAutomatedPerfTestCSVOutputMode::Separate)
-		{
-			TryStopCSVProfiler();
-		}
-	}
 }
 
 FString UAutomatedPerfTestControllerBase::GetInsightsFilename()
@@ -546,11 +510,6 @@ FString UAutomatedPerfTestControllerBase::GetCSVFilename()
 	return GetTestID();
 }
 
-void UAutomatedPerfTestControllerBase::SetCSVOutputMode(EAutomatedPerfTestCSVOutputMode NewOutputMode)
-{
-	CSVOutputMode = NewOutputMode;	
-}
-
 void UAutomatedPerfTestControllerBase::OnInit()
 {
 	Super::OnInit();
@@ -674,11 +633,6 @@ void UAutomatedPerfTestControllerBase::UnbindAllDelegates()
 	}
 }
 
-EAutomatedPerfTestCSVOutputMode UAutomatedPerfTestControllerBase::GetCSVOutputMode() const
-{
-	return CSVOutputMode;
-}
-
 void UAutomatedPerfTestControllerBase::ConsoleCommand(const TCHAR* Cmd)
 {
 	if (APlayerController* Controller = GetPlayerController())
diff --git a/Engine/Plugins/Performance/AutomatedPerfTesting/Source/AutomatedPerfTesting/Private/AutomatedReplayPerfTest.cpp b/Engine/Plugins/Performance/AutomatedPerfTesting/Source/AutomatedPerfTesting/Private/AutomatedReplayPerfTest.cpp
--- a/Engine/Plugins/Performance/AutomatedPerfTesting/Source/AutomatedPerfTesting/Private/AutomatedReplayPerfTest.cpp
+++ b/Engine/Plugins/Performance/AutomatedPerfTesting/Source/AutomatedPerfTesting/Private/AutomatedReplayPerfTest.cpp
@@ -17,7 +17,6 @@
 
 UAutomatedReplayPerfTestProjectSettings::UAutomatedReplayPerfTestProjectSettings(const FObjectInitializer& Initializer)
 	: Super(Initializer)
-	, CSVOutputMode(EAutomatedPerfTestCSVOutputMode::Separate)
 {
 }
 
@@ -119,6 +118,12 @@ void UAutomatedReplayPerfTest::TeardownTest(bool bExitAfterTeardown)
 
 	bIsTearingDown = true;
 	MarkProfilingEnd();
+
+	if (RequestsCSVProfiler())
+	{
+		TryStopCSVProfiler();
+	}
+
 	TeardownProfiling();
 	Super::TeardownTest(bExitAfterTeardown);
 }
@@ -134,7 +139,6 @@ void UAutomatedReplayPerfTest::OnInit()
 	const FString DefaultReplayExtension = TEXT(".replay");
 
 	const UAutomatedReplayPerfTestProjectSettings* Settings = GetDefault<UAutomatedReplayPerfTestProjectSettings>();
-	SetCSVOutputMode(Settings->CSVOutputMode);
 
 	const auto GetDefaultReplayFromSettings = [&]()
 	{
@@ -204,6 +208,11 @@ void UAutomatedReplayPerfTest::OnReplayStarted(UWorld* World)
 	// be deferred.
 	SetupProfiling();
 	MarkProfilingStart();
+
+	if (RequestsCSVProfiler())
+	{
+		TryStartCSVProfiler();
+	}
 }
 
 void UAutomatedReplayPerfTest::OnReplayComplete(UWorld* World)
diff --git a/Engine/Plugins/Performance/AutomatedPerfTesting/Source/AutomatedPerfTesting/Private/AutomatedSequencePerfTest.cpp b/Engine/Plugins/Performance/AutomatedPerfTesting/Source/AutomatedPerfTesting/Private/AutomatedSequencePerfTest.cpp
--- a/Engine/Plugins/Performance/AutomatedPerfTesting/Source/AutomatedPerfTesting/Private/AutomatedSequencePerfTest.cpp
+++ b/Engine/Plugins/Performance/AutomatedPerfTesting/Source/AutomatedPerfTesting/Private/AutomatedSequencePerfTest.cpp
@@ -21,7 +21,6 @@
 
 UAutomatedSequencePerfTestProjectSettings::UAutomatedSequencePerfTestProjectSettings(const FObjectInitializer& Initializer)
 	: Super(Initializer)
-	, CSVOutputMode(EAutomatedPerfTestCSVOutputMode::Separate)
 {
 }
 
@@ -162,6 +161,11 @@ void UAutomatedSequencePerfTest::RunTest()
 		UE_LOG(LogAutomatedPerfTest, Log, TEXT("RunTest::Valid Sequence Player, proceeding"));
 		MarkProfilingStart(); // Mark profiler start region on sequence start. 
 
+		if (RequestsCSVProfiler())
+		{
+			Super::TryStartCSVProfiler();
+		}
+
 		// trigger a camera cut manually in order to start the region for the first camera cut
         SequencePlayer->Play();
 		OnCameraCut(SequencePlayer->GetActiveCameraComponent());
@@ -193,6 +197,12 @@ void UAutomatedSequencePerfTest::TeardownTest(bool bExitAfterTeardown)
 	UE_LOG(LogAutomatedPerfTest, Log, TEXT("AutomatedSequencePerfTest::TeardownTest"));
 	
 	MarkProfilingEnd();
+
+	if (RequestsCSVProfiler())
+	{
+		TryStopCSVProfiler();
+	}
+
 	TeardownProfiling();
 	Super::TeardownTest(bExitAfterTeardown);
 
@@ -211,14 +221,7 @@ bool UAutomatedSequencePerfTest::TryStartCSVProfiler(FString CSVFileName, const
 #if CSV_PROFILER
 	if(FCsvProfiler* const CsvProfiler = FCsvProfiler::Get())
 	{
-		if (GetCSVOutputMode() != EAutomatedPerfTestCSVOutputMode::Single)
-		{
-			CsvProfiler->SetMetadata(TEXT("MapSequenceComboName"), *CurrentMapSequenceCombo->ComboName.ToString());
-		}
-		if (GetCSVOutputMode() == EAutomatedPerfTestCSVOutputMode::Granular)
-		{
-			CsvProfiler->SetMetadata(TEXT("CameraCut"), *GetCameraCutID());
-		}
+		CsvProfiler->SetMetadata(TEXT("MapSequenceComboName"), *CurrentMapSequenceCombo->ComboName.ToString());
 		return Super::TryStartCSVProfiler(CSVFileName, DestinationFolder, Frames);
 	}
 #endif
@@ -239,10 +242,6 @@ void UAutomatedSequencePerfTest::OnCameraCut(UCameraComponent* CameraComponent)
 		if(RequestsCSVProfiler())
 		{
 			CSV_EVENT(AutomatedPerfTest, TEXT("END_%s"), *GetCameraCutID())
-			if(GetCSVOutputMode() == EAutomatedPerfTestCSVOutputMode::Granular)
-			{
-				TryStopCSVProfiler();
-			}
 		}
 #endif
 	}
@@ -261,12 +260,7 @@ void UAutomatedSequencePerfTest::OnCameraCut(UCameraComponent* CameraComponent)
 #if CSV_PROFILER
 		if(RequestsCSVProfiler())
 		{
-			if(GetCSVOutputMode() == EAutomatedPerfTestCSVOutputMode::Granular)
-			{
-				TryStartCSVProfiler(GetCameraCutFullName());
-			}
 			CSV_EVENT(AutomatedPerfTest, TEXT("START_%s"), *GetCameraCutID())
-			
 		}
 #endif
 	}
@@ -316,8 +310,6 @@ void UAutomatedSequencePerfTest::OnInit()
 	
 	Settings = GetDefault<UAutomatedSequencePerfTestProjectSettings>();
 
-	SetCSVOutputMode(Settings->CSVOutputMode);
-	
 	// if an explicit map/sequence name was set from commandline, use this to override the test
 	if (FParse::Value(FCommandLine::Get(), TEXT("AutomatedPerfTest.SequencePerfTest.MapSequenceName="), SequenceTestName))
 	{
diff --git a/Engine/Plugins/Performance/AutomatedPerfTesting/Source/AutomatedPerfTesting/Private/StaticCameraTests/AutomatedStaticCameraPerfTestBase.cpp b/Engine/Plugins/Performance/AutomatedPerfTesting/Source/AutomatedPerfTesting/Private/StaticCameraTests/AutomatedStaticCameraPerfTestBase.cpp
--- a/Engine/Plugins/Performance/AutomatedPerfTesting/Source/AutomatedPerfTesting/Private/StaticCameraTests/AutomatedStaticCameraPerfTestBase.cpp
+++ b/Engine/Plugins/Performance/AutomatedPerfTesting/Source/AutomatedPerfTesting/Private/StaticCameraTests/AutomatedStaticCameraPerfTestBase.cpp
@@ -16,7 +16,6 @@ UAutomatedStaticCameraPerfTestProjectSettings::UAutomatedStaticCameraPerfTestPro
 	, WarmUpTime(5.0)
 	, SoakTime(5.0)
 	, CooldownTime(1.0)
-	, CSVOutputMode(EAutomatedPerfTestCSVOutputMode::Granular)
 {
 
 }
@@ -117,14 +116,8 @@ bool UAutomatedStaticCameraPerfTestBase::TryStartCSVProfiler(FString CSVFileName
 #if CSV_PROFILER
 	if(FCsvProfiler* const CsvProfiler = FCsvProfiler::Get())
 	{
-		if (GetCSVOutputMode() != EAutomatedPerfTestCSVOutputMode::Single)
-		{
-			CsvProfiler->SetMetadata(TEXT("MapName"), *CurrentMapName);
-		}
-		if (GetCSVOutputMode() == EAutomatedPerfTestCSVOutputMode::Granular)
-		{
-			CsvProfiler->SetMetadata(TEXT("CameraName"), *GetCurrentCameraRegionName());
-		}
+		CsvProfiler->SetMetadata(TEXT("MapName"), *CurrentMapName);
+		CsvProfiler->SetMetadata(TEXT("CameraName"), *GetCurrentCameraRegionName());
 		return Super::TryStartCSVProfiler(CSVFileName, DestinationFolder, Frames);
 	}
 #endif
@@ -247,11 +240,7 @@ void UAutomatedStaticCameraPerfTestBase::MarkCameraStart()
 #if CSV_PROFILER
 		if(RequestsCSVProfiler())
 		{
-			if(GetCSVOutputMode() == EAutomatedPerfTestCSVOutputMode::Granular)
-			{
-				TryStartCSVProfiler(GetCurrentCameraRegionFullName());
-			}
-			CSV_EVENT(AutomatedPerfTest, TEXT("START_%s"), *GetCurrentCameraRegionName())
+			TryStartCSVProfiler(GetCurrentCameraRegionFullName());
 		}
 #endif
 	}
@@ -269,11 +258,7 @@ void UAutomatedStaticCameraPerfTestBase::MarkCameraEnd()
 #if CSV_PROFILER
 		if(RequestsCSVProfiler())
 		{
-			CSV_EVENT(AutomatedPerfTest, TEXT("END_%s"), *GetCurrentCameraRegionName())
-			if(GetCSVOutputMode() == EAutomatedPerfTestCSVOutputMode::Granular)
-			{
-				TryStopCSVProfiler();
-			}
+			TryStopCSVProfiler();
 		}
 #endif
 	}
@@ -287,8 +272,6 @@ void UAutomatedStaticCameraPerfTestBase::OnInit()
 
 	Settings = GetDefault<UAutomatedStaticCameraPerfTestProjectSettings>();
 
-	SetCSVOutputMode(Settings->CSVOutputMode);
-	
 	// if an explicit map/sequence name was set from commandline, use this to override the test
 	if (FParse::Value(FCommandLine::Get(), TEXT("AutomatedPerfTest.StaticCameraPerfTest.MapName="), CurrentMapName))
 	{
diff --git a/Engine/Plugins/Performance/AutomatedPerfTesting/Source/AutomatedPerfTesting/Public/AutomatedMaterialPerfTest.h b/Engine/Plugins/Performance/AutomatedPerfTesting/Source/AutomatedPerfTesting/Public/AutomatedMaterialPerfTest.h
--- a/Engine/Plugins/Performance/AutomatedPerfTesting/Source/AutomatedPerfTesting/Public/AutomatedMaterialPerfTest.h
+++ b/Engine/Plugins/Performance/AutomatedPerfTesting/Source/AutomatedPerfTesting/Public/AutomatedMaterialPerfTest.h
@@ -90,12 +90,6 @@ public:
      */
 	UPROPERTY(Config, EditAnywhere, BlueprintReadWrite, Category="Material Perf Test | Plate")
 	FSoftObjectPath MaterialPlate;
-	
-	/*
-	 * For Material Perf Tests, Granular will output one CSV per-material. 
-	 */
-	UPROPERTY(Config, EditAnywhere, BlueprintReadWrite, Category="Material Perf Test | CSV")
-	EAutomatedPerfTestCSVOutputMode CSVOutputMode;
 };
 
 /**
diff --git a/Engine/Plugins/Performance/AutomatedPerfTesting/Source/AutomatedPerfTesting/Public/AutomatedPerfTestControllerBase.h b/Engine/Plugins/Performance/AutomatedPerfTesting/Source/AutomatedPerfTesting/Public/AutomatedPerfTestControllerBase.h
--- a/Engine/Plugins/Performance/AutomatedPerfTesting/Source/AutomatedPerfTesting/Public/AutomatedPerfTestControllerBase.h
+++ b/Engine/Plugins/Performance/AutomatedPerfTesting/Source/AutomatedPerfTesting/Public/AutomatedPerfTestControllerBase.h
@@ -14,15 +14,6 @@
 #define UE_API AUTOMATEDPERFTESTING_API
 #endif
 
-UENUM()
-enum class EAutomatedPerfTestCSVOutputMode : uint8
-{
-	Single UMETA(DisplayName = "Single CSV", ToolTip = "Output a single CSV with all of the results for the entire session, from SetupTest to ExitTest."),
-	Separate UMETA(DisplayName = "Separate CSVs", ToolTip = "Output CSVs from RunTest to TeardownTest. May result into multiple output CSVs that require special processing."),
-	Granular UMETA(DisplayName = "Granular CSVs", ToolTip = "Output granular CSVs during the test run, resulting in multiple CSVs between RunTest and TeardownTest.")
-};
-
-
 class AGameModeBase;
 
 namespace AutomatedPerfTest
@@ -135,10 +126,6 @@ public:
 
 	UE_API void TakeScreenshot(FString ScreenshotName);
 
-	// you'll need to set this via your subclass if you want to customize the behavior otherwise it will default to a single CSV per session
-	UE_API void SetCSVOutputMode(EAutomatedPerfTestCSVOutputMode NewOutputMode);
-	UE_API EAutomatedPerfTestCSVOutputMode GetCSVOutputMode() const;
-
 	UE_API void ConsoleCommand(const TCHAR* Cmd);
 
 protected:
@@ -190,8 +177,6 @@ private:
 	AGameModeBase* GameMode;
 
 	FDelegateHandle CsvProfilerDelegateHandle;
-
-	EAutomatedPerfTestCSVOutputMode CSVOutputMode;
 };
 
 #undef UE_API
diff --git a/Engine/Plugins/Performance/AutomatedPerfTesting/Source/AutomatedPerfTesting/Public/AutomatedReplayPerfTest.h b/Engine/Plugins/Performance/AutomatedPerfTesting/Source/AutomatedPerfTesting/Public/AutomatedReplayPerfTest.h
--- a/Engine/Plugins/Performance/AutomatedPerfTesting/Source/AutomatedPerfTesting/Public/AutomatedReplayPerfTest.h
+++ b/Engine/Plugins/Performance/AutomatedPerfTesting/Source/AutomatedPerfTesting/Public/AutomatedReplayPerfTest.h
@@ -33,9 +33,6 @@ public:
 	UPROPERTY(Config, EditAnywhere, BlueprintReadWrite, Category = "Replay Perf Test", meta = (FilePathFilter = "replay", RelativeToGameDir))
 	TArray<FFilePath>				ReplaysToTest;
 
-	UPROPERTY(Config, EditAnywhere, BlueprintReadWrite, Category = "Replay Perf Test")
-	EAutomatedPerfTestCSVOutputMode CSVOutputMode;
-
 	/* Get replay path from given replay name */
 	UFUNCTION(BlueprintCallable, Category = "Replay Perf Test")
 	UE_API bool GetReplayPathFromName(FName TestName, FString& FoundReplay) const;
diff --git a/Engine/Plugins/Performance/AutomatedPerfTesting/Source/AutomatedPerfTesting/Public/AutomatedSequencePerfTest.h b/Engine/Plugins/Performance/AutomatedPerfTesting/Source/AutomatedPerfTesting/Public/AutomatedSequencePerfTest.h
--- a/Engine/Plugins/Performance/AutomatedPerfTesting/Source/AutomatedPerfTesting/Public/AutomatedSequencePerfTest.h
+++ b/Engine/Plugins/Performance/AutomatedPerfTesting/Source/AutomatedPerfTesting/Public/AutomatedSequencePerfTest.h
@@ -72,12 +72,6 @@ public:
 	/* How long to delay between setting up the sequence for the map before the sequence actually starts */
 	UPROPERTY(Config, EditAnywhere, BlueprintReadWrite, Category="Sequence Perf Test")
 	float SequenceStartDelay=5.0;
-
-	/*
-	 * For Sequence Perf Tests, Separate will output one CSV per map tested, and Granular will output one CSV per camera-cut. 
-	 */
-	UPROPERTY(Config, EditAnywhere, BlueprintReadWrite, Category="Sequence Perf Test")
-	EAutomatedPerfTestCSVOutputMode CSVOutputMode;
 };
 
 /**
diff --git a/Engine/Plugins/Performance/AutomatedPerfTesting/Source/AutomatedPerfTesting/Public/StaticCameraTests/AutomatedStaticCameraPerfTestBase.h b/Engine/Plugins/Performance/AutomatedPerfTesting/Source/AutomatedPerfTesting/Public/StaticCameraTests/AutomatedStaticCameraPerfTestBase.h
--- a/Engine/Plugins/Performance/AutomatedPerfTesting/Source/AutomatedPerfTesting/Public/StaticCameraTests/AutomatedStaticCameraPerfTestBase.h
+++ b/Engine/Plugins/Performance/AutomatedPerfTesting/Source/AutomatedPerfTesting/Public/StaticCameraTests/AutomatedStaticCameraPerfTestBase.h
@@ -61,12 +61,6 @@ public:
      */
 	UPROPERTY(Config, EditAnywhere, BlueprintReadWrite, Category="Static Cameras")
 	float CooldownTime;
-
-	/*
-	 * For Static Camera Perf Tests, Separate will output one CSV per map tested, and Granular will output one CSV per camera. 
-	 */
-	UPROPERTY(Config, EditAnywhere, BlueprintReadWrite, Category="Static Cameras")
-	EAutomatedPerfTestCSVOutputMode CSVOutputMode;
 };
 
 /**
