Fix the automated performance test plugin running a test despite an empty material list.

#Synced-CL 49734378

diff --git a/Engine/Plugins/Performance/AutomatedPerfTesting/Source/AutomatedPerfTesting/Private/AutomatedMaterialPerfTest.cpp b/Engine/Plugins/Performance/AutomatedPerfTesting/Source/AutomatedPerfTesting/Private/AutomatedMaterialPerfTest.cpp
--- a/Engine/Plugins/Performance/AutomatedPerfTesting/Source/AutomatedPerfTesting/Private/AutomatedMaterialPerfTest.cpp
+++ b/Engine/Plugins/Performance/AutomatedPerfTesting/Source/AutomatedPerfTesting/Private/AutomatedMaterialPerfTest.cpp
@@ -238,7 +238,7 @@ void UAutomatedMaterialPerfTest::OnInit()
 	Settings = GetDefault<UAutomatedMaterialPerfTestProjectSettings>();
 
 	// early out if there aren't actually any materials
-	if(Settings->MaterialsToTest.Num() < 0)
+	if (Settings->MaterialsToTest.IsEmpty())
 	{
 		UE_LOG(LogAutomatedPerfTest, Error, TEXT("No materials defined in the project's Automated Perf Test | Materials settings. Exiting test early."));
 		EndAutomatedPerfTest(1);
