Fixing incorrect format specifiers in Automated Performance Testing code.

#Synced-CL 49734378

diff --git a/Engine/Plugins/Performance/AutomatedPerfTesting/Source/AutomatedPerfTesting/Private/AutomatedSequencePerfTest.cpp b/Engine/Plugins/Performance/AutomatedPerfTesting/Source/AutomatedPerfTesting/Private/AutomatedSequencePerfTest.cpp
--- a/Engine/Plugins/Performance/AutomatedPerfTesting/Source/AutomatedPerfTesting/Private/AutomatedSequencePerfTest.cpp
+++ b/Engine/Plugins/Performance/AutomatedPerfTesting/Source/AutomatedPerfTesting/Private/AutomatedSequencePerfTest.cpp
@@ -332,7 +332,7 @@ void UAutomatedSequencePerfTest::OnInit()
 		}
 	}
 	
-	UE_LOG(LogAutomatedPerfTest, Log, TEXT("Initialized with %hhd MapSequence combos"), MapSequenceCombos.Num());
+	UE_LOG(LogAutomatedPerfTest, Log, TEXT("Initialized with %d MapSequence combos"), MapSequenceCombos.Num());
 }
 
 void UAutomatedSequencePerfTest::UnbindAllDelegates()
diff --git a/Engine/Plugins/Performance/AutomatedPerfTesting/Source/AutomatedPerfTesting/Private/LaunchExtension/AutomatedPerfTestLaunchExtension.cpp b/Engine/Plugins/Performance/AutomatedPerfTesting/Source/AutomatedPerfTesting/Private/LaunchExtension/AutomatedPerfTestLaunchExtension.cpp
--- a/Engine/Plugins/Performance/AutomatedPerfTesting/Source/AutomatedPerfTesting/Private/LaunchExtension/AutomatedPerfTestLaunchExtension.cpp
+++ b/Engine/Plugins/Performance/AutomatedPerfTesting/Source/AutomatedPerfTesting/Private/LaunchExtension/AutomatedPerfTestLaunchExtension.cpp
@@ -222,7 +222,7 @@ void FAutomatedPerfTestLaunchExtensionInstance::CustomizeAutomatedTestCommandLin
 	};
 
 	const int32 IterationCount = GetConfigInteger(EConfig::User_PerProfile, IterationsConfig, DefaultIterationCount);
-	InOutCommandLine += FString::Printf(TEXT(" -iterations=%u"), IterationCount);
+	InOutCommandLine += FString::Printf(TEXT(" -iterations=%d"), IterationCount);
 
 	switch (CurrentTestType)
 	{
diff --git a/Engine/Plugins/Performance/AutomatedPerfTesting/Source/AutomatedPerfTesting/Private/StaticCameraTests/AutomatedStaticCameraPerfTestBase.cpp b/Engine/Plugins/Performance/AutomatedPerfTesting/Source/AutomatedPerfTesting/Private/StaticCameraTests/AutomatedStaticCameraPerfTestBase.cpp
--- a/Engine/Plugins/Performance/AutomatedPerfTesting/Source/AutomatedPerfTesting/Private/StaticCameraTests/AutomatedStaticCameraPerfTestBase.cpp
+++ b/Engine/Plugins/Performance/AutomatedPerfTesting/Source/AutomatedPerfTesting/Private/StaticCameraTests/AutomatedStaticCameraPerfTestBase.cpp
@@ -54,7 +54,7 @@ void UAutomatedStaticCameraPerfTestBase::SetupTest()
 					return;
 				}
 
-				UE_LOG(LogAutomatedPerfTest, Log, TEXT("Found %hhd cameras to test in map %s"), CamerasToTest.Num(), *CurrentMapPath.GetAssetName());
+				UE_LOG(LogAutomatedPerfTest, Log, TEXT("Found %d cameras to test in map %s"), CamerasToTest.Num(), *CurrentMapPath.GetAssetName());
 
 				GetFirstPlayerController()->SetCinematicMode(true, true, true);
 
@@ -298,7 +298,7 @@ void UAutomatedStaticCameraPerfTestBase::OnInit()
 		}
 	}
 	
-	UE_LOG(LogAutomatedPerfTest, Log, TEXT("Initialized with %hhd MapSequence combos"), MapsToTest.Num());
+	UE_LOG(LogAutomatedPerfTest, Log, TEXT("Initialized with %d MapSequence combos"), MapsToTest.Num());
 	
 	// early out if there aren't actually any maps set in project settings
 	if(Settings->MapsToTest.Num() < 0)
