Fixes a collection of code analysis identified issues in Automated Performance Testing controllers.

- Updated UAutomatedMaterialPerfTest to handle object creation failure and a potential data configuration divide-by-zero.
- Updated UAutomatedSequencePerfTest and UAutomatedStaticCameraPerfTestBase to exit initialization when test preconditions prevent the test from running.
- Removed a use of static in a UProfileGo local  which was triggering warnings about potential incorrect objects referenced. This is likely a superfluous warning but removal resolves ambiguity.
- Fixed a spurious formatting specifier in UProfileGo.

#Synced-CL 49734378

diff --git a/Engine/Plugins/Performance/AutomatedPerfTesting/Source/AutomatedPerfTesting/Private/AutomatedMaterialPerfTest.cpp b/Engine/Plugins/Performance/AutomatedPerfTesting/Source/AutomatedPerfTesting/Private/AutomatedMaterialPerfTest.cpp
--- a/Engine/Plugins/Performance/AutomatedPerfTesting/Source/AutomatedPerfTesting/Private/AutomatedMaterialPerfTest.cpp
+++ b/Engine/Plugins/Performance/AutomatedPerfTesting/Source/AutomatedPerfTesting/Private/AutomatedMaterialPerfTest.cpp
@@ -44,7 +44,14 @@ void UAutomatedMaterialPerfTest::SetupTest()
 			// Also load and verify the material plate mesh
 			UE_LOG(LogAutomatedPerfTest, Log, TEXT("Loading material plate mesh: %s"), *Settings->MaterialPlate.ToString())
 			UStaticMesh* LoadedMaterialPlateMesh = LoadObject<UStaticMesh>(NULL, *Settings->MaterialPlate.ToString(), NULL, LOAD_None, NULL);
-			check(LoadedMaterialPlateMesh);
+			if (LoadedMaterialPlateMesh == nullptr)
+			{
+				UE_LOG(LogAutomatedPerfTest, Error, TEXT("Failed to load the setting specified MaterialPlate '%s'. Exiting..."), 
+					*Settings->MaterialPlate.ToString());
+
+				EndAutomatedPerfTest(1);
+				return;
+			}
 
 			// reset the material index;
 			CurrentMaterialIndex = -1;
@@ -54,6 +61,15 @@ void UAutomatedMaterialPerfTest::SetupTest()
 			
 			// then spawn the camera into the world and set it up 
 			Camera = GetWorld()->SpawnActor<ACameraActor>();
+			if (Camera == nullptr)
+			{
+				UE_LOG(LogAutomatedPerfTest, Error, TEXT("Failed to spawn ACameraActor. Exiting..."), 
+					*Settings->MaterialPlate.ToString());
+
+				EndAutomatedPerfTest(1);
+				return;
+			}
+
 			Camera->GetCameraComponent()->SetProjectionMode(Settings->CameraProjectionMode);
 			Camera->GetCameraComponent()->SetOrthoWidth(Settings->PlateDistanceFromCamera);
 
@@ -62,7 +78,18 @@ void UAutomatedMaterialPerfTest::SetupTest()
 			MaterialPlate->SetMobility(EComponentMobility::Type::Movable);
 			MaterialPlate->GetStaticMeshComponent()->SetStaticMesh(LoadedMaterialPlateMesh);
 			MaterialPlate->SetActorLocation(FVector(Settings->PlateDistanceFromCamera, 0.0, 0.0));
-			float Scale = Settings->PlateDistanceFromCamera / (LoadedMaterialPlateMesh->GetBounds().BoxExtent.Y);
+
+			// Ensure the loaded plate width is a viable value for the division.
+			float PlateHalfWidth = LoadedMaterialPlateMesh->GetBounds().BoxExtent.Y;
+			constexpr float MinPlateHalfExtentsCm = 1.0f;
+			if (PlateHalfWidth < MinPlateHalfExtentsCm)
+			{
+				UE_LOG(LogAutomatedPerfTest, Warning, TEXT("Setting specified MaterialPlate '%s' extent '%f' are less than the minimum '%f'. Using '%f'."), 
+					*Settings->MaterialPlate.ToString(), PlateHalfWidth, MinPlateHalfExtentsCm, MinPlateHalfExtentsCm);
+				PlateHalfWidth = MinPlateHalfExtentsCm;
+			}
+
+			float Scale = Settings->PlateDistanceFromCamera / PlateHalfWidth;
 
 			UE_LOG(LogAutomatedPerfTest, Verbose, TEXT("SizeY = %f, Scale = %f"), LoadedMaterialPlateMesh->GetBounds().BoxExtent.Y, Scale);
 			
diff --git a/Engine/Plugins/Performance/AutomatedPerfTesting/Source/AutomatedPerfTesting/Private/AutomatedSequencePerfTest.cpp b/Engine/Plugins/Performance/AutomatedPerfTesting/Source/AutomatedPerfTesting/Private/AutomatedSequencePerfTest.cpp
--- a/Engine/Plugins/Performance/AutomatedPerfTesting/Source/AutomatedPerfTesting/Private/AutomatedSequencePerfTest.cpp
+++ b/Engine/Plugins/Performance/AutomatedPerfTesting/Source/AutomatedPerfTesting/Private/AutomatedSequencePerfTest.cpp
@@ -74,6 +74,7 @@ void UAutomatedSequencePerfTest::SetupTest()
 				{
 					UE_LOG(LogAutomatedPerfTest, Error, TEXT("Unable to create sequence player when starting AutomatedSequencePerfTest, exiting..."));
 					EndAutomatedPerfTest(1);
+					return;
 				}
 
 				// set the sequence up at the beginning
@@ -98,6 +99,7 @@ void UAutomatedSequencePerfTest::SetupTest()
 			{
 				UE_LOG(LogAutomatedPerfTest, Error, TEXT("Invalid World when starting AutomatedSequencePerfTest, exiting..."));
 				EndAutomatedPerfTest(1);
+				return;
 			}
 		}
 		else
diff --git a/Engine/Plugins/Performance/AutomatedPerfTesting/Source/AutomatedPerfTesting/Private/ProfileGo/ProfileGo.cpp b/Engine/Plugins/Performance/AutomatedPerfTesting/Source/AutomatedPerfTesting/Private/ProfileGo/ProfileGo.cpp
--- a/Engine/Plugins/Performance/AutomatedPerfTesting/Source/AutomatedPerfTesting/Private/ProfileGo/ProfileGo.cpp
+++ b/Engine/Plugins/Performance/AutomatedPerfTesting/Source/AutomatedPerfTesting/Private/ProfileGo/ProfileGo.cpp
@@ -332,7 +332,8 @@ bool UProfileGo::AddScenariosToCurrentRequest(const TCHAR* ProfileName, bool bEr
 	}
 	else if (bErrorIfNotFound)
 	{
-		AbortWithError(TEXT("Could not find expected scenario or collection named %s. Aborting profilego!"));
+		FString AbortMessage = FString::Printf(TEXT("Could not find expected scenario or collection named %s. Aborting profilego!"), ProfileName);
+		AbortWithError(*AbortMessage);
 	}
 
 	return false;
@@ -419,7 +420,7 @@ bool UProfileGo::WaitForLoadingAndStreaming(UWorld* InWorld, float DeltaTime, bo
 		UE_LOG(LogAutomatedPerfTest, Log, TEXT("TickSettling: Waiting On AsyncLoading"));
 	}
 
-	static TWeakObjectPtr<UWorldPartitionSubsystem> WorldPartition = UWorld::GetSubsystem<UWorldPartitionSubsystem>(InWorld);
+	TWeakObjectPtr<UWorldPartitionSubsystem> WorldPartition = UWorld::GetSubsystem<UWorldPartitionSubsystem>(InWorld);
 	// Ensure world partition streaming completes. Should be redundant with above IsStreamingInLevels()
 	if (WorldPartition.Get() && !WorldPartition->IsAllStreamingCompleted())
 	{
diff --git a/Engine/Plugins/Performance/AutomatedPerfTesting/Source/AutomatedPerfTesting/Private/StaticCameraTests/AutomatedStaticCameraPerfTestBase.cpp b/Engine/Plugins/Performance/AutomatedPerfTesting/Source/AutomatedPerfTesting/Private/StaticCameraTests/AutomatedStaticCameraPerfTestBase.cpp
--- a/Engine/Plugins/Performance/AutomatedPerfTesting/Source/AutomatedPerfTesting/Private/StaticCameraTests/AutomatedStaticCameraPerfTestBase.cpp
+++ b/Engine/Plugins/Performance/AutomatedPerfTesting/Source/AutomatedPerfTesting/Private/StaticCameraTests/AutomatedStaticCameraPerfTestBase.cpp
@@ -282,10 +282,11 @@ void UAutomatedStaticCameraPerfTestBase::OnInit()
 				MapsToTest.Add(MapPath);
 			}
 		}
-		if(MapsToTest.Num() == 0)
+		if(MapsToTest.IsEmpty())
 		{
 			UE_LOG(LogAutomatedPerfTest, Error, TEXT("Couldn't find a map name matching %s in Static Camera Maps to Test setting. Exiting."), *CurrentMapName)
 			EndAutomatedPerfTest(1);
+			return;
 		}
 	}
 	// otherwise, use all the maps defined in project settings
@@ -301,10 +302,11 @@ void UAutomatedStaticCameraPerfTestBase::OnInit()
 	UE_LOG(LogAutomatedPerfTest, Log, TEXT("Initialized with %d MapSequence combos"), MapsToTest.Num());
 	
 	// early out if there aren't actually any maps set in project settings
-	if(Settings->MapsToTest.Num() < 0)
+	if(Settings->MapsToTest.IsEmpty())
 	{
 		UE_LOG(LogAutomatedPerfTest, Error, TEXT("No maps defined in the project's Automated Perf Test | Static Camera settings. Exiting test early."));
 		EndAutomatedPerfTest(1);
+		return;
 	}
 }
 
