Update the Automated Performance Testing base controller to use a task event rather than a delegate to ensure CSV files are flushed before exit.

Previously, this waited on CSVProfiler via the OnCSVProfileFinished delegate. This had a few issues:

1. Capturing this introduced a potential dangling reference if the controller was destroyed before the lambda executed.
2. OnCSVProfileFinished may be dispatched from the CSVProcessingThread. This meant there was a potential race condition when registering the lambda, and it required both the lambda and the code it invoked to be thread-safe.

In practice, these were unlikely to cause issues due to protection elsewhere. This change switches to using graph events and explicit dispatching on the game thread to make the implementation more obviously safe.

#Synced-CL 49818450

diff --git a/Engine/Plugins/Performance/AutomatedPerfTesting/Source/AutomatedPerfTesting/Private/AutomatedPerfTestControllerBase.cpp b/Engine/Plugins/Performance/AutomatedPerfTesting/Source/AutomatedPerfTesting/Private/AutomatedPerfTestControllerBase.cpp
--- a/Engine/Plugins/Performance/AutomatedPerfTesting/Source/AutomatedPerfTesting/Private/AutomatedPerfTestControllerBase.cpp
+++ b/Engine/Plugins/Performance/AutomatedPerfTesting/Source/AutomatedPerfTesting/Private/AutomatedPerfTestControllerBase.cpp
@@ -185,8 +185,10 @@ bool UAutomatedPerfTestControllerBase::TryStartCSVProfiler(FString CSVFileName,
 		
 		CsvProfiler->BeginCapture(Frames, DestinationFolder, CSVFileName);
 		CsvProfiler->SetDeviceProfileName(GetDeviceProfile());
-		
-		return CsvProfiler->IsCapturing();
+
+		// BeginCapture queues a capture to start on a frame boundary. It isn't possible to verify the capture
+		// is now running when this function returns.
+		return true;
 	}
 #endif
 	UE_LOG(LogAutomatedPerfTest, Warning, TEXT("CSVProfiler Start requested, but not available."))
@@ -201,8 +203,10 @@ bool UAutomatedPerfTestControllerBase::TryStopCSVProfiler()
 		const bool bSafeToStopCsvProfiler = CsvProfiler->IsCapturing() && !CsvProfiler->IsEndCapturePending();
 		if (bSafeToStopCsvProfiler)
 		{
-			const FGraphEventRef AutomatedPerfTestEndEvent = FGraphEvent::CreateGraphEvent();
-			CsvProfiler->EndCapture(AutomatedPerfTestEndEvent);
+			// Store the event representing the most recent request to end a capture. Waiting on the last handle is 
+			// sufficient to ensure all files are processed.
+			CsvProfileEndCaptureEvent = FGraphEvent::CreateGraphEvent();
+			CsvProfiler->EndCapture(CsvProfileEndCaptureEvent);
 		}
 		return true;
 	}
@@ -364,26 +368,31 @@ void UAutomatedPerfTestControllerBase::Exit()
 		IAutomatedPerfTestInterface::Execute_Exit(GameMode);
 	}
 
-	if(RequestsCSVProfiler())
+	// Wait to end the performance test and request shutdown if out artifacts are still being processed.
+	FGraphEventArray Prerequisites;
+
+	// Wait for CSVProfiler to finish processing it's last file.
+	if (CsvProfileEndCaptureEvent.IsValid())
 	{
-#if CSV_PROFILER
-		if(FCsvProfiler::Get()->IsWritingFile())
-		{
-			UE_LOG(LogAutomatedPerfTest, Log, TEXT("CSVProfile requested, and test is exiting, but CSV Profiler isn't done writing."))
-			// if we requested a CSV Profile, and the CSV Profiler is still writing the file, add a lambda to call the EndAutomatedPerfTest function
-			// so that we don't exit out of the application before the CSV Profiler is done
-			
-			// TODO investigate if this is necessary if we set csv.BlockOnCaptureEnd
-			FCsvProfiler::Get()->OnCSVProfileFinished().AddLambda([this](const FString& Filename)
-			{
-				EndAutomatedPerfTest();
-			});
-			return;
-		}
-#endif
+		Prerequisites.Add(CsvProfileEndCaptureEvent);
 	}
-	
-	EndAutomatedPerfTest();
+
+	TWeakObjectPtr<UAutomatedPerfTestControllerBase> WeakThis(this);
+	FFunctionGraphTask::CreateAndDispatchWhenReady(
+	    [WeakThis]()
+	    {
+		    UAutomatedPerfTestControllerBase* This = WeakThis.Get();
+		    if (!This)
+		    {
+			    return;
+		    }
+
+		    This->EndAutomatedPerfTest();
+	    },
+	    TStatId(),
+	    &Prerequisites,
+	    ENamedThreads::GameThread
+	);
 }
 
 AGameModeBase* UAutomatedPerfTestControllerBase::GetGameMode() const
diff --git a/Engine/Plugins/Performance/AutomatedPerfTesting/Source/AutomatedPerfTesting/Public/AutomatedPerfTestControllerBase.h b/Engine/Plugins/Performance/AutomatedPerfTesting/Source/AutomatedPerfTesting/Public/AutomatedPerfTestControllerBase.h
--- a/Engine/Plugins/Performance/AutomatedPerfTesting/Source/AutomatedPerfTesting/Public/AutomatedPerfTestControllerBase.h
+++ b/Engine/Plugins/Performance/AutomatedPerfTesting/Source/AutomatedPerfTesting/Public/AutomatedPerfTestControllerBase.h
@@ -176,6 +176,10 @@ private:
 	
 	AGameModeBase* GameMode;
 
+	// Event for the most recent call to TryStopCSVProfiler. Used to ensure the resulting csv file is flushed 
+	// before exiting via EndAutomatedPerfTest().
+	FGraphEventRef CsvProfileEndCaptureEvent;
+
 	FDelegateHandle CsvProfilerDelegateHandle;
 };
 
