Replacing the Automated Performance Testing GetTestID class hierarchy based extension point with GetPerfTestTypeID to return the a class identifier.

This is the first step in moving towards data driven customization of TestID rather than code defined names. The existing generated names remain unchanged.

#Synced-CL 49896846

#Synced-CL 49968655

diff --git a/Engine/Plugins/Performance/AutomatedPerfTesting/Source/AutomatedPerfTesting/Private/AutomatedMaterialPerfTest.cpp b/Engine/Plugins/Performance/AutomatedPerfTesting/Source/AutomatedPerfTesting/Private/AutomatedMaterialPerfTest.cpp
--- a/Engine/Plugins/Performance/AutomatedPerfTesting/Source/AutomatedPerfTesting/Private/AutomatedMaterialPerfTest.cpp
+++ b/Engine/Plugins/Performance/AutomatedPerfTesting/Source/AutomatedPerfTesting/Private/AutomatedMaterialPerfTest.cpp
@@ -128,9 +128,9 @@ void UAutomatedMaterialPerfTest::RunTest()
 	GetWorld()->GetTimerManager().SetTimer(UnusedHandle, this, &UAutomatedMaterialPerfTest::SetUpNextMaterial, 1.0, false, Settings->WarmUpTime);
 }
 
-FString UAutomatedMaterialPerfTest::GetTestID()
+FString UAutomatedMaterialPerfTest::GetPerfTestTypeID() const
 {
-	return Super::GetTestID() + "_Materials";
+	return TEXT("Materials");
 }
 
 bool UAutomatedMaterialPerfTest::TryStartCSVProfiler(FString CSVFileName, const FString& DestinationFolder, int32 Frames)
diff --git a/Engine/Plugins/Performance/AutomatedPerfTesting/Source/AutomatedPerfTesting/Private/AutomatedPerfTestControllerBase.cpp b/Engine/Plugins/Performance/AutomatedPerfTesting/Source/AutomatedPerfTesting/Private/AutomatedPerfTestControllerBase.cpp
--- a/Engine/Plugins/Performance/AutomatedPerfTesting/Source/AutomatedPerfTesting/Private/AutomatedPerfTestControllerBase.cpp
+++ b/Engine/Plugins/Performance/AutomatedPerfTesting/Source/AutomatedPerfTesting/Private/AutomatedPerfTestControllerBase.cpp
@@ -81,26 +81,27 @@ UAutomatedPerfTestControllerBase::UAutomatedPerfTestControllerBase(const FObject
 	TestDatetime = FDateTime::Now().ToString(TEXT("%Y%m%d-%H%M%S"));
 }
 
-FString UAutomatedPerfTestControllerBase::GetTestName()
+FString UAutomatedPerfTestControllerBase::GetTestName() const
 {
 	return TestID;
 }
 
-FString UAutomatedPerfTestControllerBase::GetDeviceProfile()
+FString UAutomatedPerfTestControllerBase::GetDeviceProfile() const
 {
 	return DeviceProfileOverride.IsEmpty() ? UDeviceProfileManager::Get().GetActiveDeviceProfileName() : DeviceProfileOverride;  
 }
 
-FString UAutomatedPerfTestControllerBase::GetTestID()
+FString UAutomatedPerfTestControllerBase::GetTestID() const
 {
 	const TArray<FString> TestCaseIDElements = {FApp::GetBuildVersion(),
 										  FPlatformProperties::PlatformName(),
 										  TestDatetime,
-								          *GetDeviceProfile(),
-								          *GetTestName()};
+								          GetDeviceProfile(),
+								          GetTestName(),
+										  GetPerfTestTypeID()};
 
 	// TODO make this format definable in project settings, similar to how MovieRenderQueue does it for renders
-	// construct a unique ID of the form BuildVersion_PlatformName_YYYYMMDD-HHMMSS_DeviceProfile_TestName
+	// construct a unique ID of the form BuildVersion_PlatformName_YYYYMMDD-HHMMSS_DeviceProfile_TestName_TestTypeID
 	FString TestCaseID = FString::Join(TestCaseIDElements, TEXT("_"));
 
 	return TestCaseID;
@@ -131,6 +132,11 @@ bool UAutomatedPerfTestControllerBase::RequestsFPSChart() const
 	return bRequestsFPSChart;
 }
 
+FString UAutomatedPerfTestControllerBase::GetPerfTestTypeID() const
+{
+	return TEXT("");
+}
+
 bool UAutomatedPerfTestControllerBase::RequestsVideoCapture() const
 {
 	return bRequestsVideoCapture;
diff --git a/Engine/Plugins/Performance/AutomatedPerfTesting/Source/AutomatedPerfTesting/Private/AutomatedProfileGoTest.cpp b/Engine/Plugins/Performance/AutomatedPerfTesting/Source/AutomatedPerfTesting/Private/AutomatedProfileGoTest.cpp
--- a/Engine/Plugins/Performance/AutomatedPerfTesting/Source/AutomatedPerfTesting/Private/AutomatedProfileGoTest.cpp
+++ b/Engine/Plugins/Performance/AutomatedPerfTesting/Source/AutomatedPerfTesting/Private/AutomatedProfileGoTest.cpp
@@ -89,9 +89,9 @@ void UAutomatedProfileGoTest::OnInit()
 	Super::OnInit();
 }
 
-FString UAutomatedProfileGoTest::GetTestID()
+FString UAutomatedProfileGoTest::GetPerfTestTypeID() const
 {
-	return Super::GetTestID() + "_ProfileGo";
+	return TEXT("ProfileGo");
 }
 
 void UAutomatedProfileGoTest::SetupTest()
diff --git a/Engine/Plugins/Performance/AutomatedPerfTesting/Source/AutomatedPerfTesting/Private/AutomatedReplayPerfTest.cpp b/Engine/Plugins/Performance/AutomatedPerfTesting/Source/AutomatedPerfTesting/Private/AutomatedReplayPerfTest.cpp
--- a/Engine/Plugins/Performance/AutomatedPerfTesting/Source/AutomatedPerfTesting/Private/AutomatedReplayPerfTest.cpp
+++ b/Engine/Plugins/Performance/AutomatedPerfTesting/Source/AutomatedPerfTesting/Private/AutomatedReplayPerfTest.cpp
@@ -47,9 +47,9 @@ bool UAutomatedReplayPerfTestProjectSettings::GetReplayPathFromName(FName TestNa
 	return false;
 }
 
-FString UAutomatedReplayPerfTest::GetTestID()
+FString UAutomatedReplayPerfTest::GetPerfTestTypeID() const
 {
-	return Super::GetTestID() + "_Replay";
+	return TEXT("Replay");
 }
 
 void UAutomatedReplayPerfTest::SetupTest()
diff --git a/Engine/Plugins/Performance/AutomatedPerfTesting/Source/AutomatedPerfTesting/Private/AutomatedSequencePerfTest.cpp b/Engine/Plugins/Performance/AutomatedPerfTesting/Source/AutomatedPerfTesting/Private/AutomatedSequencePerfTest.cpp
--- a/Engine/Plugins/Performance/AutomatedPerfTesting/Source/AutomatedPerfTesting/Private/AutomatedSequencePerfTest.cpp
+++ b/Engine/Plugins/Performance/AutomatedPerfTesting/Source/AutomatedPerfTesting/Private/AutomatedSequencePerfTest.cpp
@@ -268,9 +268,9 @@ void UAutomatedSequencePerfTest::OnCameraCut(UCameraComponent* CameraComponent)
 	}
 }
 
-FString UAutomatedSequencePerfTest::GetTestID()
+FString UAutomatedSequencePerfTest::GetPerfTestTypeID() const
 {
-	return Super::GetTestID() + "_Sequence";
+	return TEXT("Sequence");
 }
 
 FString UAutomatedSequencePerfTest::GetCameraCutID()
diff --git a/Engine/Plugins/Performance/AutomatedPerfTesting/Source/AutomatedPerfTesting/Private/StaticCameraTests/AutomatedStaticCameraPerfTestBase.cpp b/Engine/Plugins/Performance/AutomatedPerfTesting/Source/AutomatedPerfTesting/Private/StaticCameraTests/AutomatedStaticCameraPerfTestBase.cpp
--- a/Engine/Plugins/Performance/AutomatedPerfTesting/Source/AutomatedPerfTesting/Private/StaticCameraTests/AutomatedStaticCameraPerfTestBase.cpp
+++ b/Engine/Plugins/Performance/AutomatedPerfTesting/Source/AutomatedPerfTesting/Private/StaticCameraTests/AutomatedStaticCameraPerfTestBase.cpp
@@ -96,9 +96,9 @@ void UAutomatedStaticCameraPerfTestBase::RunTest()
 	GetWorld()->GetTimerManager().SetTimer(UnusedHandle, this, &UAutomatedStaticCameraPerfTestBase::SetUpNextCamera, 1.0, false, Settings->WarmUpTime);
 }
 
-FString UAutomatedStaticCameraPerfTestBase::GetTestID()
+FString UAutomatedStaticCameraPerfTestBase::GetPerfTestTypeID() const
 {
-	return Super::GetTestID() + "_StaticCamera";
+	return TEXT("StaticCamera");
 }
 
 FString UAutomatedStaticCameraPerfTestBase::GetCSVFilename()
diff --git a/Engine/Plugins/Performance/AutomatedPerfTesting/Source/AutomatedPerfTesting/Public/AutomatedMaterialPerfTest.h b/Engine/Plugins/Performance/AutomatedPerfTesting/Source/AutomatedPerfTesting/Public/AutomatedMaterialPerfTest.h
--- a/Engine/Plugins/Performance/AutomatedPerfTesting/Source/AutomatedPerfTesting/Public/AutomatedMaterialPerfTest.h
+++ b/Engine/Plugins/Performance/AutomatedPerfTesting/Source/AutomatedPerfTesting/Public/AutomatedMaterialPerfTest.h
@@ -107,7 +107,7 @@ public:
 	UFUNCTION()
 	UE_API virtual void RunTest() override;
 
-	UE_API virtual FString GetTestID() override;
+	UE_API virtual FString GetPerfTestTypeID() const override;
 
 	UE_API virtual bool TryStartCSVProfiler(FString CSVFileName, const FString& DestinationFolder = FString(), int32 Frames = -1) override;
 	// ~End UAutomatedPerfTestControllerBase Interface
diff --git a/Engine/Plugins/Performance/AutomatedPerfTesting/Source/AutomatedPerfTesting/Public/AutomatedPerfTestControllerBase.h b/Engine/Plugins/Performance/AutomatedPerfTesting/Source/AutomatedPerfTesting/Public/AutomatedPerfTestControllerBase.h
--- a/Engine/Plugins/Performance/AutomatedPerfTesting/Source/AutomatedPerfTesting/Public/AutomatedPerfTestControllerBase.h
+++ b/Engine/Plugins/Performance/AutomatedPerfTesting/Source/AutomatedPerfTesting/Public/AutomatedPerfTestControllerBase.h
@@ -88,9 +88,13 @@ public:
 public:
 	UE_API UAutomatedPerfTestControllerBase(const FObjectInitializer& ObjectInitializer);
 
-	UE_API FString GetTestName();
-	UE_API FString GetDeviceProfile();
-	UE_API virtual FString GetTestID();
+	UE_API FString GetTestName() const;
+	UE_API FString GetDeviceProfile() const;
+	UE_API FString GetTestID() const;
+
+	// Unique identifier for the class type provided to output artifact generation.
+	UE_API virtual FString GetPerfTestTypeID() const;
+
 	UE_API FString GetOverallRegionName();
 	UE_API FString GetTraceChannels();
 	
diff --git a/Engine/Plugins/Performance/AutomatedPerfTesting/Source/AutomatedPerfTesting/Public/AutomatedProfileGoTest.h b/Engine/Plugins/Performance/AutomatedPerfTesting/Source/AutomatedPerfTesting/Public/AutomatedProfileGoTest.h
--- a/Engine/Plugins/Performance/AutomatedPerfTesting/Source/AutomatedPerfTesting/Public/AutomatedProfileGoTest.h
+++ b/Engine/Plugins/Performance/AutomatedPerfTesting/Source/AutomatedPerfTesting/Public/AutomatedProfileGoTest.h
@@ -31,7 +31,7 @@ public:
 	UE_API UAutomatedProfileGoTest(const FObjectInitializer& ObjectInitializer);
 
 	UE_API virtual void OnInit() override;
-	UE_API virtual FString GetTestID() override;
+	UE_API virtual FString GetPerfTestTypeID() const override;
 	UE_API virtual void SetupTest() override;
 	UE_API virtual void OnTick(float TimeDelta) override;
 
diff --git a/Engine/Plugins/Performance/AutomatedPerfTesting/Source/AutomatedPerfTesting/Public/AutomatedReplayPerfTest.h b/Engine/Plugins/Performance/AutomatedPerfTesting/Source/AutomatedPerfTesting/Public/AutomatedReplayPerfTest.h
--- a/Engine/Plugins/Performance/AutomatedPerfTesting/Source/AutomatedPerfTesting/Public/AutomatedReplayPerfTest.h
+++ b/Engine/Plugins/Performance/AutomatedPerfTesting/Source/AutomatedPerfTesting/Public/AutomatedReplayPerfTest.h
@@ -48,7 +48,7 @@ class UAutomatedReplayPerfTest : public UAutomatedPerfTestControllerBase
 
 public:
 	// ~Begin UAutomatedPerfTestControllerBase Interface
-	UE_API virtual FString GetTestID() override;
+	UE_API virtual FString GetPerfTestTypeID() const override;
 	UE_API virtual void SetupTest() override;
 
 	UFUNCTION()
diff --git a/Engine/Plugins/Performance/AutomatedPerfTesting/Source/AutomatedPerfTesting/Public/AutomatedSequencePerfTest.h b/Engine/Plugins/Performance/AutomatedPerfTesting/Source/AutomatedPerfTesting/Public/AutomatedSequencePerfTest.h
--- a/Engine/Plugins/Performance/AutomatedPerfTesting/Source/AutomatedPerfTesting/Public/AutomatedSequencePerfTest.h
+++ b/Engine/Plugins/Performance/AutomatedPerfTesting/Source/AutomatedPerfTesting/Public/AutomatedSequencePerfTest.h
@@ -84,7 +84,7 @@ class UAutomatedSequencePerfTest : public UAutomatedPerfTestControllerBase
 
 public:
 	// ~Begin UAutomatedPerfTestControllerBase Interface
-	UE_API virtual FString GetTestID() override;
+	UE_API virtual FString GetPerfTestTypeID() const override;
 	UE_API virtual void SetupTest() override;
 	
 	UFUNCTION()
diff --git a/Engine/Plugins/Performance/AutomatedPerfTesting/Source/AutomatedPerfTesting/Public/StaticCameraTests/AutomatedStaticCameraPerfTestBase.h b/Engine/Plugins/Performance/AutomatedPerfTesting/Source/AutomatedPerfTesting/Public/StaticCameraTests/AutomatedStaticCameraPerfTestBase.h
--- a/Engine/Plugins/Performance/AutomatedPerfTesting/Source/AutomatedPerfTesting/Public/StaticCameraTests/AutomatedStaticCameraPerfTestBase.h
+++ b/Engine/Plugins/Performance/AutomatedPerfTesting/Source/AutomatedPerfTesting/Public/StaticCameraTests/AutomatedStaticCameraPerfTestBase.h
@@ -78,7 +78,7 @@ public:
 	UFUNCTION()
 	UE_API virtual void RunTest() override;
 
-	UE_API virtual FString GetTestID() override;
+	UE_API virtual FString GetPerfTestTypeID() const override;
 
 	UE_API virtual bool TryStartCSVProfiler(FString CSVFileName, const FString& DestinationFolder = FString(), int32 Frames = -1) override;
 	// ~End UAutomatedPerfTestControllerBase Interface
