Updated Automated Performance Testing BuildGraph scripts to discover and apply AddPerfTestConfigsm enabling generation of performance test variants.

This supports defining multiple permutations of a given build target for automated performance testing. As an example, Samples/Showcases/CitySample/Platforms/Windows/Build/PerfTestConfigs.xml defines an additional Windows configuration that adds Vulkan RHI performance testing alongside the default Windows RHI configuration.

#Synced-CL 50067330

diff --git a/Engine/Plugins/Performance/AutomatedPerfTesting/Build/Scripts/AutomatedPerfTestNode.cs b/Engine/Plugins/Performance/AutomatedPerfTesting/Build/Scripts/AutomatedPerfTestNode.cs
--- a/Engine/Plugins/Performance/AutomatedPerfTesting/Build/Scripts/AutomatedPerfTestNode.cs
+++ b/Engine/Plugins/Performance/AutomatedPerfTesting/Build/Scripts/AutomatedPerfTestNode.cs
@@ -1339,6 +1339,11 @@ namespace AutomatedPerfTest
 				"testconfigname=" + Config.TestConfigName,
 			];
 
+			if (!string.IsNullOrEmpty(Config.TestConfigName))
+			{
+				CsvMetadata.Add("deviceconfigoverride=" + Config.TestConfigName);
+			}
+
 			if (!string.IsNullOrEmpty(Context.BuildInfo.Branch) && Context.BuildInfo.Changelist != 0)
 			{
 				CsvMetadata.Add("branch=" + Context.BuildInfo.Branch);
diff --git a/Engine/Plugins/Performance/AutomatedPerfTesting/Resources/TemplatesForHookup/AutoPerfTests.xml b/Engine/Plugins/Performance/AutomatedPerfTesting/Resources/TemplatesForHookup/AutoPerfTests.xml
--- a/Engine/Plugins/Performance/AutomatedPerfTesting/Resources/TemplatesForHookup/AutoPerfTests.xml
+++ b/Engine/Plugins/Performance/AutomatedPerfTesting/Resources/TemplatesForHookup/AutoPerfTests.xml
@@ -18,6 +18,14 @@
 	<Property Name="DefaultTargetTestList" Value="" />
 </Do>
 
+<!-- Customized the project configurations to enable project-specific performance testing using the 
+AutomatedPerfTesting/Build/Inc/TestConfigs provided macro 'AddPerfTestConfig'. 
+
+For example, on Windows this can be used to validate non-default RHIs such as Vulkan. -->
+<Include Script="$(RootDir)/Engine/Plugins/Performance/AutomatedPerfTesting/Build/Inc/TestConfigs.xml" />
+<Include Script="../Platforms/*/Build/PerfTestConfigs.xml" />
+<Include Script="../Restricted/*/Build/PerfTestConfigs.xml" />
+
 <!-- Configure AutomatedPerfTest Sequence test -->
 <Include Script="$(RootDir)/Engine/Plugins/Performance/AutomatedPerfTesting/Build/Inc/AutomatedPerfTestProjectSettings.xml" />
 <Do If="'$(MapSequenceComboNames)' == ''">
