Updated the automated performance test controller base class to expose an overridable metadata collection method.

Added UAutomatedPerfTestControllerBase::GatherTestMetadata to explicitly gather properties passed to the CSV profiler, replacing the previous pattern of overriding TryStartCSVProfiler. This simplifies the controller implementations and moves closer to establishing an abstraction of profiling instruments for APT, making it easier to introduce additional instruments through code or data-driven configuration.

#Synced-CL 50067330

diff --git a/Engine/Plugins/Performance/AutomatedPerfTesting/Source/AutomatedPerfTesting/Private/AutomatedMaterialPerfTest.cpp b/Engine/Plugins/Performance/AutomatedPerfTesting/Source/AutomatedPerfTesting/Private/AutomatedMaterialPerfTest.cpp
--- a/Engine/Plugins/Performance/AutomatedPerfTesting/Source/AutomatedPerfTesting/Private/AutomatedMaterialPerfTest.cpp
+++ b/Engine/Plugins/Performance/AutomatedPerfTesting/Source/AutomatedPerfTesting/Private/AutomatedMaterialPerfTest.cpp
@@ -133,17 +133,11 @@ FString UAutomatedMaterialPerfTest::GetPerfTestTypeID() const
 	return TEXT("Materials");
 }
 
-bool UAutomatedMaterialPerfTest::TryStartCSVProfiler(FString CSVFileName, const FString& DestinationFolder, int32 Frames)
+void UAutomatedMaterialPerfTest::GatherTestMetadata(TArray<TPair<FString, FString>>& OutMetadata) const
 {
-#if CSV_PROFILER
-	if(FCsvProfiler* const CsvProfiler = FCsvProfiler::Get())
-	{
-		CsvProfiler->SetMetadata(TEXT("Material"), *GetCurrentMaterialRegionName());
-		return Super::TryStartCSVProfiler(CSVFileName, DestinationFolder, Frames);
-	}
-#endif
-	UE_LOG(LogAutomatedPerfTest, Warning, TEXT("CSVProfiler Start requested, but not available."))
-	return false;
+	Super::GatherTestMetadata(OutMetadata);
+
+	OutMetadata.Emplace(TEXT("Material"), GetCurrentMaterialRegionName());
 }
 
 void UAutomatedMaterialPerfTest::SetUpNextMaterial()
@@ -210,7 +204,7 @@ UMaterialInterface* UAutomatedMaterialPerfTest::GetCurrentMaterial() const
 	return CurrentMaterial;
 }
 
-FString UAutomatedMaterialPerfTest::GetCurrentMaterialRegionName()
+FString UAutomatedMaterialPerfTest::GetCurrentMaterialRegionName() const
 {
 	return GetCurrentMaterial()->GetName();
 }
diff --git a/Engine/Plugins/Performance/AutomatedPerfTesting/Source/AutomatedPerfTesting/Private/AutomatedPerfTestControllerBase.cpp b/Engine/Plugins/Performance/AutomatedPerfTesting/Source/AutomatedPerfTesting/Private/AutomatedPerfTestControllerBase.cpp
--- a/Engine/Plugins/Performance/AutomatedPerfTesting/Source/AutomatedPerfTesting/Private/AutomatedPerfTestControllerBase.cpp
+++ b/Engine/Plugins/Performance/AutomatedPerfTesting/Source/AutomatedPerfTesting/Private/AutomatedPerfTestControllerBase.cpp
@@ -169,6 +169,14 @@ bool UAutomatedPerfTestControllerBase::TryStartCSVProfiler()
 	return TryStartCSVProfiler(GetCSVFilename(), DestinationFolder);
 }
 
+void UAutomatedPerfTestControllerBase::GatherTestMetadata(TArray<TPair<FString, FString>>& OutMetadata) const
+{
+	OutMetadata.Emplace(TEXT("TestID"), TestID);
+	OutMetadata.Emplace(TEXT("Datetime"), TestDatetime);
+	OutMetadata.Emplace(TEXT("ResX"), FString::FromInt(GSystemResolution.ResX));
+	OutMetadata.Emplace(TEXT("ResY"), FString::FromInt(GSystemResolution.ResY));
+}
+
 bool UAutomatedPerfTestControllerBase::TryStartCSVProfiler(FString CSVFileName, const FString& DestinationFolder, int32 Frames)
 {
 #if CSV_PROFILER
@@ -181,10 +189,12 @@ bool UAutomatedPerfTestControllerBase::TryStartCSVProfiler(FString CSVFileName,
 
 		UE_LOG(LogAutomatedPerfTest, Log, TEXT("Attempting to start CSV Profile to file %s"), *CSVFileName);
 
-		CsvProfiler->SetMetadata(TEXT("TestID"), *TestID);
-		CsvProfiler->SetMetadata(TEXT("Datetime"), *TestDatetime);
-		CsvProfiler->SetMetadata(TEXT("ResX"), *FString::FromInt(GSystemResolution.ResX));
-		CsvProfiler->SetMetadata(TEXT("ResY"), *FString::FromInt(GSystemResolution.ResY));
+		TArray<TPair<FString, FString>> Metadata;
+		GatherTestMetadata(Metadata);
+		for (const TPair<FString, FString>& MetadataPair : Metadata)
+		{
+			CsvProfiler->SetMetadata(*MetadataPair.Key, *MetadataPair.Value);
+		}
 
 		UE_LOG(LogGauntlet, Log, TEXT("CSV Profiler Destination Filename: %s, Destination Folder: %s"), 
 				*CSVFileName, DestinationFolder.IsEmpty()? TEXT("<Default>") : *DestinationFolder);
diff --git a/Engine/Plugins/Performance/AutomatedPerfTesting/Source/AutomatedPerfTesting/Private/AutomatedSequencePerfTest.cpp b/Engine/Plugins/Performance/AutomatedPerfTesting/Source/AutomatedPerfTesting/Private/AutomatedSequencePerfTest.cpp
--- a/Engine/Plugins/Performance/AutomatedPerfTesting/Source/AutomatedPerfTesting/Private/AutomatedSequencePerfTest.cpp
+++ b/Engine/Plugins/Performance/AutomatedPerfTesting/Source/AutomatedPerfTesting/Private/AutomatedSequencePerfTest.cpp
@@ -165,7 +165,7 @@ void UAutomatedSequencePerfTest::RunTest()
 
 		if (RequestsCSVProfiler())
 		{
-			Super::TryStartCSVProfiler();
+			TryStartCSVProfiler();
 		}
 
 		// trigger a camera cut manually in order to start the region for the first camera cut
@@ -218,17 +218,17 @@ void UAutomatedSequencePerfTest::TeardownTest(bool bExitAfterTeardown)
 	NextMap();
 }
 
-bool UAutomatedSequencePerfTest::TryStartCSVProfiler(FString CSVFileName, const FString& DestinationFolder, int32 Frames)
+void UAutomatedSequencePerfTest::GatherTestMetadata(TArray<TPair<FString, FString>>& OutMetadata) const
 {
-#if CSV_PROFILER
-	if(FCsvProfiler* const CsvProfiler = FCsvProfiler::Get())
-	{
-		CsvProfiler->SetMetadata(TEXT("MapSequenceComboName"), *CurrentMapSequenceCombo->ComboName.ToString());
-		return Super::TryStartCSVProfiler(CSVFileName, DestinationFolder, Frames);
-	}
-#endif
-	UE_LOG(LogAutomatedPerfTest, Warning, TEXT("CSVProfiler Start requested, but not available."))
-	return false;
+	Super::GatherTestMetadata(OutMetadata);
+
+	// Called after a test is created. CurrentMapSequenceCombo and ComboName are required invariants of this class's 
+	// implementation. Missing values result in malformed output.
+	checkf(CurrentMapSequenceCombo.IsSet(), TEXT("GatherTestMetadata requires a CurrentMapSequenceCombo."));
+	FString ComboName = CurrentMapSequenceCombo->ComboName.ToString();
+	checkf(!ComboName.IsEmpty(), TEXT("GatherTestMetadata requires a valid ComboName."));
+
+	OutMetadata.Emplace(TEXT("MapSequenceComboName"), ComboName);
 }
 
 void UAutomatedSequencePerfTest::OnCameraCut(UCameraComponent* CameraComponent)
diff --git a/Engine/Plugins/Performance/AutomatedPerfTesting/Source/AutomatedPerfTesting/Private/StaticCameraTests/AutomatedStaticCameraPerfTestBase.cpp b/Engine/Plugins/Performance/AutomatedPerfTesting/Source/AutomatedPerfTesting/Private/StaticCameraTests/AutomatedStaticCameraPerfTestBase.cpp
--- a/Engine/Plugins/Performance/AutomatedPerfTesting/Source/AutomatedPerfTesting/Private/StaticCameraTests/AutomatedStaticCameraPerfTestBase.cpp
+++ b/Engine/Plugins/Performance/AutomatedPerfTesting/Source/AutomatedPerfTesting/Private/StaticCameraTests/AutomatedStaticCameraPerfTestBase.cpp
@@ -111,18 +111,12 @@ FString UAutomatedStaticCameraPerfTestBase::GetCSVFilename()
 	return GetTestID() + "_" + GetCurrentMap();
 }
 
-bool UAutomatedStaticCameraPerfTestBase::TryStartCSVProfiler(FString CSVFileName, const FString& DestinationFolder, int32 Frames)
+void UAutomatedStaticCameraPerfTestBase::GatherTestMetadata(TArray<TPair<FString, FString>>& OutMetadata) const
 {
-#if CSV_PROFILER
-	if(FCsvProfiler* const CsvProfiler = FCsvProfiler::Get())
-	{
-		CsvProfiler->SetMetadata(TEXT("MapName"), *CurrentMapName);
-		CsvProfiler->SetMetadata(TEXT("CameraName"), *GetCurrentCameraRegionName());
-		return Super::TryStartCSVProfiler(CSVFileName, DestinationFolder, Frames);
-	}
-#endif
-	UE_LOG(LogAutomatedPerfTest, Warning, TEXT("CSVProfiler Start requested, but not available."))
-	return false;
+	Super::GatherTestMetadata(OutMetadata);
+
+	OutMetadata.Emplace(TEXT("MapName"), CurrentMapName);
+	OutMetadata.Emplace(TEXT("CameraName"), GetCurrentCameraRegionName());
 }
 
 void UAutomatedStaticCameraPerfTestBase::SetUpNextCamera()
@@ -218,7 +212,7 @@ ACameraActor* UAutomatedStaticCameraPerfTestBase::GetCurrentCamera() const
 	return CurrentCamera;
 }
 
-FString UAutomatedStaticCameraPerfTestBase::GetCurrentCameraRegionName()
+FString UAutomatedStaticCameraPerfTestBase::GetCurrentCameraRegionName() const
 {
 	return GetCurrentCamera()->GetActorNameOrLabel();
 }
diff --git a/Engine/Plugins/Performance/AutomatedPerfTesting/Source/AutomatedPerfTesting/Public/AutomatedMaterialPerfTest.h b/Engine/Plugins/Performance/AutomatedPerfTesting/Source/AutomatedPerfTesting/Public/AutomatedMaterialPerfTest.h
--- a/Engine/Plugins/Performance/AutomatedPerfTesting/Source/AutomatedPerfTesting/Public/AutomatedMaterialPerfTest.h
+++ b/Engine/Plugins/Performance/AutomatedPerfTesting/Source/AutomatedPerfTesting/Public/AutomatedMaterialPerfTest.h
@@ -109,7 +109,7 @@ public:
 
 	UE_API virtual FString GetPerfTestTypeID() const override;
 
-	UE_API virtual bool TryStartCSVProfiler(FString CSVFileName, const FString& DestinationFolder = FString(), int32 Frames = -1) override;
+	UE_API virtual void GatherTestMetadata(TArray<TPair<FString, FString>>& OutMetadata) const override;
 	// ~End UAutomatedPerfTestControllerBase Interface
 	
 	UFUNCTION()
@@ -125,7 +125,7 @@ public:
 	UE_API void ScreenshotMaterial();
 
 	UE_API UMaterialInterface* GetCurrentMaterial() const;
-	UE_API FString GetCurrentMaterialRegionName();
+	UE_API FString GetCurrentMaterialRegionName() const;
 	UE_API FString GetCurrentMaterialRegionFullName();
 
 	UE_API void MarkMaterialStart();
diff --git a/Engine/Plugins/Performance/AutomatedPerfTesting/Source/AutomatedPerfTesting/Public/AutomatedPerfTestControllerBase.h b/Engine/Plugins/Performance/AutomatedPerfTesting/Source/AutomatedPerfTesting/Public/AutomatedPerfTestControllerBase.h
--- a/Engine/Plugins/Performance/AutomatedPerfTesting/Source/AutomatedPerfTesting/Public/AutomatedPerfTestControllerBase.h
+++ b/Engine/Plugins/Performance/AutomatedPerfTesting/Source/AutomatedPerfTesting/Public/AutomatedPerfTestControllerBase.h
@@ -107,8 +107,8 @@ public:
 	UE_API bool TryStartInsightsTrace();
 	UE_API bool TryStopInsightsTrace();
 
-	UE_API virtual bool TryStartCSVProfiler();
-	UE_API virtual bool TryStartCSVProfiler(FString CSVFileName, const FString& DestinationFolder = FString(), int32 Frames = -1);
+	UE_API bool TryStartCSVProfiler();
+	UE_API bool TryStartCSVProfiler(FString CSVFileName, const FString& DestinationFolder = FString(), int32 Frames = -1);
 	UE_API bool TryStopCSVProfiler();
 	
 	UE_API bool TryStartFPSChart();
@@ -117,6 +117,12 @@ public:
 	UE_API bool TryStartVideoCapture();
 	UE_API bool TryFinalizingVideoCapture(const bool bStopAutoContinue = false);
 
+	// Gather test metadata key-value pairs for profiling instruments.
+	//
+	// Override this method to add controller specific metadata (map names, material names, etc). Always call 
+	// Super::GatherTestMetadata() to include base metadata.
+	UE_API virtual void GatherTestMetadata(TArray<TPair<FString, FString>>& OutMetadata) const;
+
 	UE_API virtual void SetupTest();
 	UE_API virtual void RunTest();
 	UE_API virtual void TeardownTest(bool bExitAfterTeardown = true);
diff --git a/Engine/Plugins/Performance/AutomatedPerfTesting/Source/AutomatedPerfTesting/Public/AutomatedSequencePerfTest.h b/Engine/Plugins/Performance/AutomatedPerfTesting/Source/AutomatedPerfTesting/Public/AutomatedSequencePerfTest.h
--- a/Engine/Plugins/Performance/AutomatedPerfTesting/Source/AutomatedPerfTesting/Public/AutomatedSequencePerfTest.h
+++ b/Engine/Plugins/Performance/AutomatedPerfTesting/Source/AutomatedPerfTesting/Public/AutomatedSequencePerfTest.h
@@ -92,7 +92,7 @@ public:
 	
 	UE_API virtual void TeardownTest(bool bExitAfterTeardown = true) override;
 
-	UE_API virtual bool TryStartCSVProfiler(FString CSVFileName, const FString& DestinationFolder = FString(), int32 Frames = -1) override;
+	UE_API void GatherTestMetadata(TArray<TPair<FString, FString>>& OutMetadata) const override;
 	// ~End UAutomatedPerfTestControllerBase Interface
 	
 	// This function is called on world change to set up the map for the correct map/sequence combo
diff --git a/Engine/Plugins/Performance/AutomatedPerfTesting/Source/AutomatedPerfTesting/Public/StaticCameraTests/AutomatedStaticCameraPerfTestBase.h b/Engine/Plugins/Performance/AutomatedPerfTesting/Source/AutomatedPerfTesting/Public/StaticCameraTests/AutomatedStaticCameraPerfTestBase.h
--- a/Engine/Plugins/Performance/AutomatedPerfTesting/Source/AutomatedPerfTesting/Public/StaticCameraTests/AutomatedStaticCameraPerfTestBase.h
+++ b/Engine/Plugins/Performance/AutomatedPerfTesting/Source/AutomatedPerfTesting/Public/StaticCameraTests/AutomatedStaticCameraPerfTestBase.h
@@ -80,7 +80,7 @@ public:
 
 	UE_API virtual FString GetPerfTestTypeID() const override;
 
-	UE_API virtual bool TryStartCSVProfiler(FString CSVFileName, const FString& DestinationFolder = FString(), int32 Frames = -1) override;
+	UE_API virtual void GatherTestMetadata(TArray<TPair<FString, FString>>& OutMetadata) const override;
 	// ~End UAutomatedPerfTestControllerBase Interface
 	
 	UFUNCTION()
@@ -101,7 +101,7 @@ public:
 	UE_API virtual TArray<ACameraActor*> GetMapCameraActors();
 	
 	UE_API ACameraActor* GetCurrentCamera() const;
-	UE_API FString GetCurrentCameraRegionName();
+	UE_API FString GetCurrentCameraRegionName() const;
 	UE_API FString GetCurrentCameraRegionFullName();
 
 	UE_API void MarkCameraStart();
