Simplify Automated Performance Testing controllers by consolidating the TryStartCSVProfiler overloads, standardizing CSV destination folder selection rules, and updating each controller to directly select the output filename.

Previously, destination path selection was inconsistent across controller types. Replay and sequence performance tests honored AutomatedPerfTest.OverrideArtifactPath, while others did not. Controllers now apply a consistent order of preference for the output directory:

1. Function argument (explicit output directory)
2. Command line override
3. Default value

The base class also previously provided an extension point for retrieving the CSV filename. In practice, this added complexity without sufficient value. For example, UAutomatedStaticCameraPerfTestBase implemented GetCSVFilename(), but invoked the TryStartCSVProfiler() overload that takes a filename argument and instead passed GetCurrentCameraRegionFullName(). Controllers now directly provide the filename when starting the CSV profiler, removing the redundant filename extension point and making intent explicit at call sites.

#Synced-CL 50067330

diff --git a/Engine/Plugins/Performance/AutomatedPerfTesting/Source/AutomatedPerfTesting/Private/AutomatedMaterialPerfTest.cpp b/Engine/Plugins/Performance/AutomatedPerfTesting/Source/AutomatedPerfTesting/Private/AutomatedMaterialPerfTest.cpp
--- a/Engine/Plugins/Performance/AutomatedPerfTesting/Source/AutomatedPerfTesting/Private/AutomatedMaterialPerfTest.cpp
+++ b/Engine/Plugins/Performance/AutomatedPerfTesting/Source/AutomatedPerfTesting/Private/AutomatedMaterialPerfTest.cpp
@@ -209,11 +209,6 @@ FString UAutomatedMaterialPerfTest::GetCurrentMaterialRegionName() const
 	return GetCurrentMaterial()->GetName();
 }
 
-FString UAutomatedMaterialPerfTest::GetCurrentMaterialRegionFullName()
-{
-	return GetTestID() + "_" + GetCurrentMaterialRegionName();
-}
-
 void UAutomatedMaterialPerfTest::MarkMaterialStart()
 {
 	// safety check on the current material
@@ -226,7 +221,8 @@ void UAutomatedMaterialPerfTest::MarkMaterialStart()
 #if CSV_PROFILER
 		if(RequestsCSVProfiler())
 		{
-			TryStartCSVProfiler(GetCurrentMaterialRegionFullName());
+			FString CSVFilename = GetTestID() + "_" + GetCurrentMaterialRegionName();
+			TryStartCSVProfiler(CSVFilename);
 		}
 #endif
 	}
diff --git a/Engine/Plugins/Performance/AutomatedPerfTesting/Source/AutomatedPerfTesting/Private/AutomatedPerfTestControllerBase.cpp b/Engine/Plugins/Performance/AutomatedPerfTesting/Source/AutomatedPerfTesting/Private/AutomatedPerfTestControllerBase.cpp
--- a/Engine/Plugins/Performance/AutomatedPerfTesting/Source/AutomatedPerfTesting/Private/AutomatedPerfTestControllerBase.cpp
+++ b/Engine/Plugins/Performance/AutomatedPerfTesting/Source/AutomatedPerfTesting/Private/AutomatedPerfTestControllerBase.cpp
@@ -163,12 +163,6 @@ bool UAutomatedPerfTestControllerBase::TryStopInsightsTrace()
 	return false;
 }
 
-bool UAutomatedPerfTestControllerBase::TryStartCSVProfiler()
-{
-	const FString DestinationFolder = ArtifactOutputPath.IsEmpty() ? FString() : FPaths::Combine(ArtifactOutputPath, "CSV");
-	return TryStartCSVProfiler(GetCSVFilename(), DestinationFolder);
-}
-
 void UAutomatedPerfTestControllerBase::GatherTestMetadata(TArray<TPair<FString, FString>>& OutMetadata) const
 {
 	OutMetadata.Emplace(TEXT("TestID"), TestID);
@@ -196,10 +190,25 @@ bool UAutomatedPerfTestControllerBase::TryStartCSVProfiler(FString CSVFileName,
 			CsvProfiler->SetMetadata(*MetadataPair.Key, *MetadataPair.Value);
 		}
 
+		// Select the destination folder for csv files. 
+		FString CSVDestinationFolder;
+		if (!DestinationFolder.IsEmpty())
+		{
+			CSVDestinationFolder = DestinationFolder;
+		}
+		else if (!ArtifactOutputPath.IsEmpty())
+		{
+			CSVDestinationFolder = FPaths::Combine(ArtifactOutputPath, "CSV");
+		}
+		else
+		{
+			CSVDestinationFolder = "";
+		}
+
 		UE_LOG(LogGauntlet, Log, TEXT("CSV Profiler Destination Filename: %s, Destination Folder: %s"), 
-				*CSVFileName, DestinationFolder.IsEmpty()? TEXT("<Default>") : *DestinationFolder);
+				*CSVFileName, CSVDestinationFolder.IsEmpty()? TEXT("<Default>") : *CSVDestinationFolder);
 		
-		CsvProfiler->BeginCapture(Frames, DestinationFolder, CSVFileName);
+		CsvProfiler->BeginCapture(Frames, CSVDestinationFolder, CSVFileName);
 		CsvProfiler->SetDeviceProfileName(GetDeviceProfile());
 
 		// BeginCapture queues a capture to start on a frame boundary. It isn't possible to verify the capture
@@ -530,11 +539,6 @@ FString UAutomatedPerfTestControllerBase::GetInsightsFilename()
 	return GetTestID();
 }
 
-FString UAutomatedPerfTestControllerBase::GetCSVFilename()
-{
-	return GetTestID();
-}
-
 void UAutomatedPerfTestControllerBase::OnInit()
 {
 	Super::OnInit();
diff --git a/Engine/Plugins/Performance/AutomatedPerfTesting/Source/AutomatedPerfTesting/Private/AutomatedReplayPerfTest.cpp b/Engine/Plugins/Performance/AutomatedPerfTesting/Source/AutomatedPerfTesting/Private/AutomatedReplayPerfTest.cpp
--- a/Engine/Plugins/Performance/AutomatedPerfTesting/Source/AutomatedPerfTesting/Private/AutomatedReplayPerfTest.cpp
+++ b/Engine/Plugins/Performance/AutomatedPerfTesting/Source/AutomatedPerfTesting/Private/AutomatedReplayPerfTest.cpp
@@ -211,7 +211,7 @@ void UAutomatedReplayPerfTest::OnReplayStarted(UWorld* World)
 
 	if (RequestsCSVProfiler())
 	{
-		TryStartCSVProfiler();
+		TryStartCSVProfiler(GetTestID());
 	}
 }
 
diff --git a/Engine/Plugins/Performance/AutomatedPerfTesting/Source/AutomatedPerfTesting/Private/AutomatedSequencePerfTest.cpp b/Engine/Plugins/Performance/AutomatedPerfTesting/Source/AutomatedPerfTesting/Private/AutomatedSequencePerfTest.cpp
--- a/Engine/Plugins/Performance/AutomatedPerfTesting/Source/AutomatedPerfTesting/Private/AutomatedSequencePerfTest.cpp
+++ b/Engine/Plugins/Performance/AutomatedPerfTesting/Source/AutomatedPerfTesting/Private/AutomatedSequencePerfTest.cpp
@@ -165,7 +165,18 @@ void UAutomatedSequencePerfTest::RunTest()
 
 		if (RequestsCSVProfiler())
 		{
-			TryStartCSVProfiler();
+			FString CSVFilename;
+			if (CurrentMapSequenceCombo.IsSet() && !CurrentMapSequenceCombo->ComboName.ToString().IsEmpty())
+			{
+				CSVFilename = GetTestID() + "_" + CurrentMapSequenceCombo->ComboName.ToString();
+			}
+			else
+			{
+				CSVFilename = GetTestID();
+				UE_LOG(LogAutomatedPerfTest, Error, TEXT("Current Map Sequence Combo Name not set"));
+			}
+
+			TryStartCSVProfiler(CSVFilename);
 		}
 
 		// trigger a camera cut manually in order to start the region for the first camera cut
@@ -360,14 +371,3 @@ void UAutomatedSequencePerfTest::UnbindAllDelegates()
 		World->GetTimerManager().ClearAllTimersForObject(this);
 	}
 }
-
-FString UAutomatedSequencePerfTest::GetCSVFilename()
-{
-	if (!CurrentMapSequenceCombo.IsSet() || CurrentMapSequenceCombo->ComboName.ToString().IsEmpty())
-	{
-		UE_LOG(LogAutomatedPerfTest, Error, TEXT("Current Map Sequence Combo Name not set"));
-		return Super::GetCSVFilename();
-	}
-
-	return GetTestID() + "_" + CurrentMapSequenceCombo->ComboName.ToString();
-}
diff --git a/Engine/Plugins/Performance/AutomatedPerfTesting/Source/AutomatedPerfTesting/Private/ProfileGo/ProfileGo.cpp b/Engine/Plugins/Performance/AutomatedPerfTesting/Source/AutomatedPerfTesting/Private/ProfileGo/ProfileGo.cpp
--- a/Engine/Plugins/Performance/AutomatedPerfTesting/Source/AutomatedPerfTesting/Private/ProfileGo/ProfileGo.cpp
+++ b/Engine/Plugins/Performance/AutomatedPerfTesting/Source/AutomatedPerfTesting/Private/ProfileGo/ProfileGo.cpp
@@ -1286,7 +1286,7 @@ void UProfileGo::RegisterInternalCommandDelegates()
 
 			if (UAutomatedPerfTestControllerBase* const Controller = ProfileGoSubsystem->GetTestController())
 			{
-				FileName = Controller ? Controller->GetCSVFilename() + FileName : FileName;
+				FileName = Controller ? Controller->GetTestID() + FileName : FileName;
 				Controller->TryStartCSVProfiler(FileName, DestinationFolder, CaptureFrames);
 			}
 			else
diff --git a/Engine/Plugins/Performance/AutomatedPerfTesting/Source/AutomatedPerfTesting/Private/StaticCameraTests/AutomatedStaticCameraPerfTestBase.cpp b/Engine/Plugins/Performance/AutomatedPerfTesting/Source/AutomatedPerfTesting/Private/StaticCameraTests/AutomatedStaticCameraPerfTestBase.cpp
--- a/Engine/Plugins/Performance/AutomatedPerfTesting/Source/AutomatedPerfTesting/Private/StaticCameraTests/AutomatedStaticCameraPerfTestBase.cpp
+++ b/Engine/Plugins/Performance/AutomatedPerfTesting/Source/AutomatedPerfTesting/Private/StaticCameraTests/AutomatedStaticCameraPerfTestBase.cpp
@@ -101,16 +101,6 @@ FString UAutomatedStaticCameraPerfTestBase::GetPerfTestTypeID() const
 	return TEXT("StaticCamera");
 }
 
-FString UAutomatedStaticCameraPerfTestBase::GetCSVFilename()
-{
-	if (GetCurrentMap() != CurrentMapPath.GetAssetName())
-	{
-		UE_LOG(LogAutomatedPerfTest, Warning, TEXT("Current Map is not the expected path. Current: %s, Expected: %s"), *GetCurrentMap(), *CurrentMapPath.GetAssetName());
-	}
-	
-	return GetTestID() + "_" + GetCurrentMap();
-}
-
 void UAutomatedStaticCameraPerfTestBase::GatherTestMetadata(TArray<TPair<FString, FString>>& OutMetadata) const
 {
 	Super::GatherTestMetadata(OutMetadata);
@@ -217,11 +207,6 @@ FString UAutomatedStaticCameraPerfTestBase::GetCurrentCameraRegionName() const
 	return GetCurrentCamera()->GetActorNameOrLabel();
 }
 
-FString UAutomatedStaticCameraPerfTestBase::GetCurrentCameraRegionFullName()
-{
-	return GetTestID() + "_" + GetCurrentCameraRegionName();
-}
-
 void UAutomatedStaticCameraPerfTestBase::MarkCameraStart()
 {
 	// safety check on the current material
@@ -234,7 +219,8 @@ void UAutomatedStaticCameraPerfTestBase::MarkCameraStart()
 #if CSV_PROFILER
 		if(RequestsCSVProfiler())
 		{
-			TryStartCSVProfiler(GetCurrentCameraRegionFullName());
+			FString CSVFilename = GetTestID() + "_" + GetCurrentCameraRegionName();
+			TryStartCSVProfiler(CSVFilename);
 		}
 #endif
 	}
diff --git a/Engine/Plugins/Performance/AutomatedPerfTesting/Source/AutomatedPerfTesting/Public/AutomatedMaterialPerfTest.h b/Engine/Plugins/Performance/AutomatedPerfTesting/Source/AutomatedPerfTesting/Public/AutomatedMaterialPerfTest.h
--- a/Engine/Plugins/Performance/AutomatedPerfTesting/Source/AutomatedPerfTesting/Public/AutomatedMaterialPerfTest.h
+++ b/Engine/Plugins/Performance/AutomatedPerfTesting/Source/AutomatedPerfTesting/Public/AutomatedMaterialPerfTest.h
@@ -126,7 +126,6 @@ public:
 
 	UE_API UMaterialInterface* GetCurrentMaterial() const;
 	UE_API FString GetCurrentMaterialRegionName() const;
-	UE_API FString GetCurrentMaterialRegionFullName();
 
 	UE_API void MarkMaterialStart();
 	UE_API void MarkMaterialEnd();
diff --git a/Engine/Plugins/Performance/AutomatedPerfTesting/Source/AutomatedPerfTesting/Public/AutomatedPerfTestControllerBase.h b/Engine/Plugins/Performance/AutomatedPerfTesting/Source/AutomatedPerfTesting/Public/AutomatedPerfTestControllerBase.h
--- a/Engine/Plugins/Performance/AutomatedPerfTesting/Source/AutomatedPerfTesting/Public/AutomatedPerfTestControllerBase.h
+++ b/Engine/Plugins/Performance/AutomatedPerfTesting/Source/AutomatedPerfTesting/Public/AutomatedPerfTestControllerBase.h
@@ -107,7 +107,6 @@ public:
 	UE_API bool TryStartInsightsTrace();
 	UE_API bool TryStopInsightsTrace();
 
-	UE_API bool TryStartCSVProfiler();
 	UE_API bool TryStartCSVProfiler(FString CSVFileName, const FString& DestinationFolder = FString(), int32 Frames = -1);
 	UE_API bool TryStopCSVProfiler();
 	
@@ -130,7 +129,6 @@ public:
 	UE_API virtual void Exit();
 
 	UE_API virtual FString GetInsightsFilename();
-	UE_API virtual FString GetCSVFilename();
 
 	UE_API AGameModeBase* GetGameMode() const;
 
diff --git a/Engine/Plugins/Performance/AutomatedPerfTesting/Source/AutomatedPerfTesting/Public/AutomatedSequencePerfTest.h b/Engine/Plugins/Performance/AutomatedPerfTesting/Source/AutomatedPerfTesting/Public/AutomatedSequencePerfTest.h
--- a/Engine/Plugins/Performance/AutomatedPerfTesting/Source/AutomatedPerfTesting/Public/AutomatedSequencePerfTest.h
+++ b/Engine/Plugins/Performance/AutomatedPerfTesting/Source/AutomatedPerfTesting/Public/AutomatedSequencePerfTest.h
@@ -111,7 +111,6 @@ public:
 protected:
 	UE_API virtual void OnInit() override;
 	UE_API virtual void UnbindAllDelegates() override;
-	UE_API virtual FString GetCSVFilename() override;
 
 private:
 	const UAutomatedSequencePerfTestProjectSettings* Settings;
diff --git a/Engine/Plugins/Performance/AutomatedPerfTesting/Source/AutomatedPerfTesting/Public/StaticCameraTests/AutomatedStaticCameraPerfTestBase.h b/Engine/Plugins/Performance/AutomatedPerfTesting/Source/AutomatedPerfTesting/Public/StaticCameraTests/AutomatedStaticCameraPerfTestBase.h
--- a/Engine/Plugins/Performance/AutomatedPerfTesting/Source/AutomatedPerfTesting/Public/StaticCameraTests/AutomatedStaticCameraPerfTestBase.h
+++ b/Engine/Plugins/Performance/AutomatedPerfTesting/Source/AutomatedPerfTesting/Public/StaticCameraTests/AutomatedStaticCameraPerfTestBase.h
@@ -102,7 +102,6 @@ public:
 	
 	UE_API ACameraActor* GetCurrentCamera() const;
 	UE_API FString GetCurrentCameraRegionName() const;
-	UE_API FString GetCurrentCameraRegionFullName();
 
 	UE_API void MarkCameraStart();
 	UE_API void MarkCameraEnd();
@@ -112,8 +111,6 @@ protected:
 	UE_API virtual void UnbindAllDelegates() override;
 	UE_API virtual void TriggerExitAfterDelay() override;
 
-	UE_API virtual FString GetCSVFilename() override;
-
 private:
 	UPROPERTY()
 	TArray<TObjectPtr<ACameraActor>> CamerasToTest;
