Remove the unused Automated Performance Testing function UAutomatedSequencePerfTest::GetCameraCutFullName()

#Synced-CL 50212691

diff --git a/Engine/Plugins/Performance/AutomatedPerfTesting/Source/AutomatedPerfTesting/Private/AutomatedSequencePerfTest.cpp b/Engine/Plugins/Performance/AutomatedPerfTesting/Source/AutomatedPerfTesting/Private/AutomatedSequencePerfTest.cpp
--- a/Engine/Plugins/Performance/AutomatedPerfTesting/Source/AutomatedPerfTesting/Private/AutomatedSequencePerfTest.cpp
+++ b/Engine/Plugins/Performance/AutomatedPerfTesting/Source/AutomatedPerfTesting/Private/AutomatedSequencePerfTest.cpp
@@ -306,16 +306,6 @@ FString UAutomatedSequencePerfTest::GetCameraCutID()
 	return "";
 }
 
-FString UAutomatedSequencePerfTest::GetCameraCutFullName()
-{
-	// the full name for the camera cut should include the TestID so that the data can be 
-	if(CurrentCamera != nullptr)
-	{
-		return GetTestID() + "_" + GetCameraCutID(); 
-	}
-	return GetTestID();
-}
-
 void UAutomatedSequencePerfTest::OnInit()
 {
 	Super::OnInit();
diff --git a/Engine/Plugins/Performance/AutomatedPerfTesting/Source/AutomatedPerfTesting/Public/AutomatedSequencePerfTest.h b/Engine/Plugins/Performance/AutomatedPerfTesting/Source/AutomatedPerfTesting/Public/AutomatedSequencePerfTest.h
--- a/Engine/Plugins/Performance/AutomatedPerfTesting/Source/AutomatedPerfTesting/Public/AutomatedSequencePerfTest.h
+++ b/Engine/Plugins/Performance/AutomatedPerfTesting/Source/AutomatedPerfTesting/Public/AutomatedSequencePerfTest.h
@@ -41,7 +41,6 @@ struct FAutomatedPerfTestMapSequenceCombo
 	 */
 	UPROPERTY(Config, EditAnywhere, BlueprintReadWrite, Category="Sequence Perf Test", meta=(AllowedClasses="/Script/GameFramework.GameModeBase"))
 	FString GameModeOverride;
-	
 };
 
 
@@ -106,7 +105,6 @@ public:
 	UE_API void OnCameraCut(UCameraComponent* CameraComponent);
 	
 	UE_API FString GetCameraCutID();
-	UE_API FString GetCameraCutFullName();
 	
 protected:
 	UE_API virtual void OnInit() override;
