Inlining the Automated Performance Testing function UAutomatedMaterialPerfTest::GetCurrentMaterial() which was called a single time internally.

#Synced-CL 50212691

diff --git a/Engine/Plugins/Performance/AutomatedPerfTesting/Source/AutomatedPerfTesting/Private/AutomatedMaterialPerfTest.cpp b/Engine/Plugins/Performance/AutomatedPerfTesting/Source/AutomatedPerfTesting/Private/AutomatedMaterialPerfTest.cpp
--- a/Engine/Plugins/Performance/AutomatedPerfTesting/Source/AutomatedPerfTesting/Private/AutomatedMaterialPerfTest.cpp
+++ b/Engine/Plugins/Performance/AutomatedPerfTesting/Source/AutomatedPerfTesting/Private/AutomatedMaterialPerfTest.cpp
@@ -199,14 +199,9 @@ void UAutomatedMaterialPerfTest::ScreenshotMaterial()
 	GetWorld()->GetTimerManager().SetTimer(UnusedHandle, this, &UAutomatedMaterialPerfTest::SetUpNextMaterial, 1.0, false, .1);
 }
 
-UMaterialInterface* UAutomatedMaterialPerfTest::GetCurrentMaterial() const
-{
-	return CurrentMaterial;
-}
-
 FString UAutomatedMaterialPerfTest::GetCurrentMaterialRegionName() const
 {
-	return GetCurrentMaterial()->GetName();
+	return CurrentMaterial->GetName();
 }
 
 void UAutomatedMaterialPerfTest::MarkMaterialStart()
diff --git a/Engine/Plugins/Performance/AutomatedPerfTesting/Source/AutomatedPerfTesting/Public/AutomatedMaterialPerfTest.h b/Engine/Plugins/Performance/AutomatedPerfTesting/Source/AutomatedPerfTesting/Public/AutomatedMaterialPerfTest.h
--- a/Engine/Plugins/Performance/AutomatedPerfTesting/Source/AutomatedPerfTesting/Public/AutomatedMaterialPerfTest.h
+++ b/Engine/Plugins/Performance/AutomatedPerfTesting/Source/AutomatedPerfTesting/Public/AutomatedMaterialPerfTest.h
@@ -124,7 +124,6 @@ public:
 	UFUNCTION()
 	UE_API void ScreenshotMaterial();
 
-	UE_API UMaterialInterface* GetCurrentMaterial() const;
 	UE_API FString GetCurrentMaterialRegionName() const;
 
 	UE_API void MarkMaterialStart();
