Inlining the Automated Performance Testing function UAutomatedPerfTestControllerBase::GetInsightsFilename as it was called in a single internal place.

Simplifying the code will make it simpler to refactor the defintion of the APT capture instruments.

#Synced-CL 50212691

diff --git a/Engine/Plugins/Performance/AutomatedPerfTesting/Source/AutomatedPerfTesting/Private/AutomatedPerfTestControllerBase.cpp b/Engine/Plugins/Performance/AutomatedPerfTesting/Source/AutomatedPerfTesting/Private/AutomatedPerfTestControllerBase.cpp
--- a/Engine/Plugins/Performance/AutomatedPerfTesting/Source/AutomatedPerfTesting/Private/AutomatedPerfTestControllerBase.cpp
+++ b/Engine/Plugins/Performance/AutomatedPerfTesting/Source/AutomatedPerfTesting/Private/AutomatedPerfTestControllerBase.cpp
@@ -149,7 +149,7 @@ bool UAutomatedPerfTestControllerBase::RequestsLockedDynRes() const
 
 bool UAutomatedPerfTestControllerBase::TryStartInsightsTrace()
 {
-	const FString TraceFileName = GetInsightsFilename() + ".utrace";
+	const FString TraceFileName = GetTestID() + ".utrace";
 	UE_LOG(LogAutomatedPerfTest, Log, TEXT("Attempting to start insights trace to file %s with channels %s"), *TraceFileName, *GetTraceChannels());
 	return FTraceAuxiliary::Start(FTraceAuxiliary::EConnectionType::File, *TraceFileName, *GetTraceChannels());
 }
@@ -534,11 +534,6 @@ void UAutomatedPerfTestControllerBase::MarkProfilingEnd()
 	}
 }
 
-FString UAutomatedPerfTestControllerBase::GetInsightsFilename()
-{
-	return GetTestID();
-}
-
 void UAutomatedPerfTestControllerBase::OnInit()
 {
 	Super::OnInit();
diff --git a/Engine/Plugins/Performance/AutomatedPerfTesting/Source/AutomatedPerfTesting/Public/AutomatedPerfTestControllerBase.h b/Engine/Plugins/Performance/AutomatedPerfTesting/Source/AutomatedPerfTesting/Public/AutomatedPerfTestControllerBase.h
--- a/Engine/Plugins/Performance/AutomatedPerfTesting/Source/AutomatedPerfTesting/Public/AutomatedPerfTestControllerBase.h
+++ b/Engine/Plugins/Performance/AutomatedPerfTesting/Source/AutomatedPerfTesting/Public/AutomatedPerfTestControllerBase.h
@@ -128,8 +128,6 @@ public:
 	UE_API virtual void TriggerExitAfterDelay();
 	UE_API virtual void Exit();
 
-	UE_API virtual FString GetInsightsFilename();
-
 	UE_API AGameModeBase* GetGameMode() const;
 
 	UE_API void TakeScreenshot(FString ScreenshotName);
