[QA Automation Tools] Implement an InGame support for the new upload tool
#rnx
#Synced-CL 50152832

diff --git a/Engine/Plugins/Performance/AutomatedPerfTesting/Source/AutomatedPerfTesting/Private/ProfileGo/ProfileGo.cpp b/Engine/Plugins/Performance/AutomatedPerfTesting/Source/AutomatedPerfTesting/Private/ProfileGo/ProfileGo.cpp
--- a/Engine/Plugins/Performance/AutomatedPerfTesting/Source/AutomatedPerfTesting/Private/ProfileGo/ProfileGo.cpp
+++ b/Engine/Plugins/Performance/AutomatedPerfTesting/Source/AutomatedPerfTesting/Private/ProfileGo/ProfileGo.cpp
@@ -1129,12 +1129,6 @@ void UProfileGo::ExecuteProfileGoCommand(const FProfileGoCommandAPT& ProfileGoCo
 
 	if (TryHandleProfileGoCommandString(ProfileGoCommand) == false)
 	{
-		// TODO This name replacement should happen in just one place
-		FProfileGoScenarioResultAPT& CurrentResult = CurrentRequest->CurrentResult;
-		FString CommandString = ProfileGoCommand.Command;
-		FString ScenarioName = FPaths::MakeValidFileName(CurrentResult.Scenario.Name);
-		ScenarioName = ScenarioName.Replace(TEXT(" "), TEXT("_"));
-		CommandString.ReplaceInline(TEXT("{Location}"), *ScenarioName);
 		// regular console command
 #if PLATFORM_IOS
 		if (CommandString.Contains(TEXT("profilegpu")) == false && CommandString.Contains(TEXT("screenshot")) == false)
@@ -1142,11 +1136,11 @@ void UProfileGo::ExecuteProfileGoCommand(const FProfileGoCommandAPT& ProfileGoCo
 		{
 			if (!ProfileGoCommand.CopyOutputToGameLog && ResultsFile)
 			{
-				GEngine->Exec(GetWorld(), *CommandString, *ResultsFile);
+				GEngine->Exec(GetWorld(), *ProfileGoCommand.Command, *ResultsFile);
 			}
 			else
 			{
-				ConsoleCommand(*CommandString);
+				ConsoleCommand(*ProfileGoCommand.Command);
 			}
 		}
 	}
@@ -1192,6 +1186,8 @@ bool UProfileGo::TryHandleProfileGoCommandString(const FProfileGoCommandAPT& Pro
 	ScenarioName = ScenarioName.Replace(TEXT(" "), TEXT("_"));
 	CommandString.ReplaceInline(TEXT("{Location}"), *ScenarioName);
 
+	const_cast<FProfileGoCommandAPT&>(ProfileGoCommand).Command = CommandString;
+
 	const auto FindAndExecuteHandler = [&](bool& bOutCommandResult)
 	{
 		TPair<FString, CommandHandlerDelegate>* Found = Algo::FindByPredicate(CommandHandlers, 
