[IOS] Fix CIS issue in Profile.cpp

A previous CL removed code, but didn't update a section inside PLATFORM_IOS that was still using it

#Synced-CL 50292855

diff --git a/Engine/Plugins/Performance/AutomatedPerfTesting/Source/AutomatedPerfTesting/Private/ProfileGo/ProfileGo.cpp b/Engine/Plugins/Performance/AutomatedPerfTesting/Source/AutomatedPerfTesting/Private/ProfileGo/ProfileGo.cpp
--- a/Engine/Plugins/Performance/AutomatedPerfTesting/Source/AutomatedPerfTesting/Private/ProfileGo/ProfileGo.cpp
+++ b/Engine/Plugins/Performance/AutomatedPerfTesting/Source/AutomatedPerfTesting/Private/ProfileGo/ProfileGo.cpp
@@ -1131,7 +1131,7 @@ void UProfileGo::ExecuteProfileGoCommand(const FProfileGoCommandAPT& ProfileGoCo
 	{
 		// regular console command
 #if PLATFORM_IOS
-		if (CommandString.Contains(TEXT("profilegpu")) == false && CommandString.Contains(TEXT("screenshot")) == false)
+		if (ProfileGoCommand.Command.Contains(TEXT("profilegpu")) == false && ProfileGoCommand.Command.Contains(TEXT("screenshot")) == false)
 #endif
 		{
 			if (!ProfileGoCommand.CopyOutputToGameLog && ResultsFile)
