Updated controller subclasses now call EndTestSuccess() and EndTestFailure() to queue test exits instead of calling EndAutomatedPerfTest directly with an exit code.

Consolidating exit handling this way standardizes result semantics (0 = success, 1 = failure) and removes the previous ambiguity where both 1 and -1 were used for failures. EndTestFailure() still accepts a custom exit value for the rare cases where a distinct non-zero code is required.

Because the previous default-argument style for EndTestFailure was not compatible with the APT ProfileGo call, that call was updated to use a lambda. We now create the timer delegate with FTimerDelegate::CreateWeakLambda to safely handle the case where UAutomatedProfileGoTest is destroyed before the timer fires.

#Synced-CL 50242944

diff --git a/Engine/Plugins/Performance/AutomatedPerfTesting/Source/AutomatedPerfTesting/Private/AutomatedMaterialPerfTest.cpp b/Engine/Plugins/Performance/AutomatedPerfTesting/Source/AutomatedPerfTesting/Private/AutomatedMaterialPerfTest.cpp
--- a/Engine/Plugins/Performance/AutomatedPerfTesting/Source/AutomatedPerfTesting/Private/AutomatedMaterialPerfTest.cpp
+++ b/Engine/Plugins/Performance/AutomatedPerfTesting/Source/AutomatedPerfTesting/Private/AutomatedMaterialPerfTest.cpp
@@ -49,7 +49,7 @@ void UAutomatedMaterialPerfTest::SetupTest()
 				UE_LOG(LogAutomatedPerfTest, Error, TEXT("Failed to load the setting specified MaterialPlate '%s'. Exiting..."), 
 					*Settings->MaterialPlate.ToString());
 
-				EndAutomatedPerfTest(1);
+				EndTestFailure();
 				return;
 			}
 
@@ -66,7 +66,7 @@ void UAutomatedMaterialPerfTest::SetupTest()
 				UE_LOG(LogAutomatedPerfTest, Error, TEXT("Failed to spawn ACameraActor. Exiting..."), 
 					*Settings->MaterialPlate.ToString());
 
-				EndAutomatedPerfTest(1);
+				EndTestFailure();
 				return;
 			}
 
@@ -105,7 +105,7 @@ void UAutomatedMaterialPerfTest::SetupTest()
 		else
 		{
 			UE_LOG(LogAutomatedPerfTest, Error, TEXT("Invalid World when starting UAutomatedMaterialPerfTest, exiting..."));
-			EndAutomatedPerfTest(1);
+			EndTestFailure();
 		}
 	}
 	else
@@ -253,7 +253,7 @@ void UAutomatedMaterialPerfTest::OnInit()
 	if (Settings->MaterialsToTest.IsEmpty())
 	{
 		UE_LOG(LogAutomatedPerfTest, Error, TEXT("No materials defined in the project's Automated Perf Test | Materials settings. Exiting test early."));
-		EndAutomatedPerfTest(1);
+		EndTestFailure();
 	}
 }
 
diff --git a/Engine/Plugins/Performance/AutomatedPerfTesting/Source/AutomatedPerfTesting/Private/AutomatedPerfTestControllerBase.cpp b/Engine/Plugins/Performance/AutomatedPerfTesting/Source/AutomatedPerfTesting/Private/AutomatedPerfTestControllerBase.cpp
--- a/Engine/Plugins/Performance/AutomatedPerfTesting/Source/AutomatedPerfTesting/Private/AutomatedPerfTestControllerBase.cpp
+++ b/Engine/Plugins/Performance/AutomatedPerfTesting/Source/AutomatedPerfTesting/Private/AutomatedPerfTestControllerBase.cpp
@@ -412,7 +412,8 @@ void UAutomatedPerfTestControllerBase::Exit()
 			    return;
 		    }
 
-		    This->EndAutomatedPerfTest();
+			// Exit the test, requesting the process exit.
+		    This->EndTestSuccess();
 	    },
 	    TStatId(),
 	    &Prerequisites,
@@ -597,6 +598,16 @@ void UAutomatedPerfTestControllerBase::BeginDestroy()
 	Super::BeginDestroy();
 }
 
+void UAutomatedPerfTestControllerBase::EndTestSuccess() 
+{ 
+	EndAutomatedPerfTest(0); 
+}
+
+void UAutomatedPerfTestControllerBase::EndTestFailure(const int32 ExitCode) 
+{ 
+	EndAutomatedPerfTest(ExitCode); 
+}
+
 void UAutomatedPerfTestControllerBase::EndAutomatedPerfTest(const int32 ExitCode)
 {
 	UnbindAllDelegates();
diff --git a/Engine/Plugins/Performance/AutomatedPerfTesting/Source/AutomatedPerfTesting/Private/AutomatedProfileGoTest.cpp b/Engine/Plugins/Performance/AutomatedPerfTesting/Source/AutomatedPerfTesting/Private/AutomatedProfileGoTest.cpp
--- a/Engine/Plugins/Performance/AutomatedPerfTesting/Source/AutomatedPerfTesting/Private/AutomatedProfileGoTest.cpp
+++ b/Engine/Plugins/Performance/AutomatedPerfTesting/Source/AutomatedPerfTesting/Private/AutomatedProfileGoTest.cpp
@@ -265,10 +265,21 @@ void UAutomatedProfileGoTest::CompleteTest(bool bSuccess)
 {
 	constexpr float RateInSeconds = 5.0f;
 
+	// Handle this potentially being destroyed priorty to the timer expiring.
+	FTimerDelegate Delegate = FTimerDelegate::CreateWeakLambda(this, [this, bSuccess]()
+	{
+		if (bSuccess)
+		{
+			EndTestSuccess();
+		}
+		else
+		{
+			EndTestFailure();
+		}
+	});
+
 	FTimerHandle WaitTimerHandle;
-	FTimerDelegate::TMethodPtr<UAutomatedProfileGoTest> TimerMethod = bSuccess ? 
-		&UAutomatedProfileGoTest::EndTestSuccess : &UAutomatedProfileGoTest::EndTestFailure;
-	GWorld->GetTimerManager().SetTimer(WaitTimerHandle, this, TimerMethod, RateInSeconds, false);
+	GWorld->GetTimerManager().SetTimer(WaitTimerHandle, Delegate, RateInSeconds, false);
 	bShouldTick = false;
 }
 
diff --git a/Engine/Plugins/Performance/AutomatedPerfTesting/Source/AutomatedPerfTesting/Private/AutomatedReplayPerfTest.cpp b/Engine/Plugins/Performance/AutomatedPerfTesting/Source/AutomatedPerfTesting/Private/AutomatedReplayPerfTest.cpp
--- a/Engine/Plugins/Performance/AutomatedPerfTesting/Source/AutomatedPerfTesting/Private/AutomatedReplayPerfTest.cpp
+++ b/Engine/Plugins/Performance/AutomatedPerfTesting/Source/AutomatedPerfTesting/Private/AutomatedReplayPerfTest.cpp
@@ -82,7 +82,7 @@ void UAutomatedReplayPerfTest::SetupTest()
 	else
 	{
 		UE_LOG(LogAutomatedPerfTest, Error, TEXT("Invalid World when starting AutomatedReplayPerfTest, exiting..."));
-		EndAutomatedPerfTest(1);
+		EndTestFailure();
 	}
 }
 
@@ -101,7 +101,7 @@ void UAutomatedReplayPerfTest::RunTest()
 	if (!bIsReplayTriggered)
 	{
 		UE_LOG(LogAutomatedPerfTest, Error, TEXT("Could not start Replay Perf Test"));
-		EndAutomatedPerfTest(1);
+		EndTestFailure();
 	}
 }
 
@@ -186,7 +186,8 @@ void UAutomatedReplayPerfTest::OnInit()
 	if (ReplayName.IsEmpty() || !bReplayExists)
 	{
 		UE_LOG(LogAutomatedPerfTest, Error, TEXT("Replay not specified in args nor found in settings"));
-		EndAutomatedPerfTest(1);
+		EndTestFailure();
+		return;
 	}
 
 	UE_LOG(LogAutomatedPerfTest, Log, TEXT("Replay Name: %s"), *ReplayName);
@@ -228,6 +229,5 @@ void UAutomatedReplayPerfTest::OnReplayFailure(UWorld* World, const UE::Net::TNe
 	UE_LOG(LogAutomatedPerfTest, Error, TEXT("Replay playback error: %s"), *Error.GetErrorContext());
 
 	// End replay run with failure error code. 
-	constexpr int32 ErrorCode = 1;
-	EndAutomatedPerfTest(ErrorCode);
+	EndTestFailure();
 }
diff --git a/Engine/Plugins/Performance/AutomatedPerfTesting/Source/AutomatedPerfTesting/Private/AutomatedSequencePerfTest.cpp b/Engine/Plugins/Performance/AutomatedPerfTesting/Source/AutomatedPerfTesting/Private/AutomatedSequencePerfTest.cpp
--- a/Engine/Plugins/Performance/AutomatedPerfTesting/Source/AutomatedPerfTesting/Private/AutomatedSequencePerfTest.cpp
+++ b/Engine/Plugins/Performance/AutomatedPerfTesting/Source/AutomatedPerfTesting/Private/AutomatedSequencePerfTest.cpp
@@ -73,7 +73,7 @@ void UAutomatedSequencePerfTest::SetupTest()
 				if (SequencePlayer == nullptr)
 				{
 					UE_LOG(LogAutomatedPerfTest, Error, TEXT("Unable to create sequence player when starting AutomatedSequencePerfTest, exiting..."));
-					EndAutomatedPerfTest(1);
+					EndTestFailure();
 					return;
 				}
 
@@ -98,7 +98,7 @@ void UAutomatedSequencePerfTest::SetupTest()
 			else
 			{
 				UE_LOG(LogAutomatedPerfTest, Error, TEXT("Invalid World when starting AutomatedSequencePerfTest, exiting..."));
-				EndAutomatedPerfTest(1);
+				EndTestFailure();
 				return;
 			}
 		}
@@ -193,7 +193,7 @@ void UAutomatedSequencePerfTest::RunTest()
 	else
 	{
 		UE_LOG(LogAutomatedPerfTest, Error, TEXT("Invalid SequencePlayer when starting AutomatedSequencePerfTest, exiting..."));
-		EndAutomatedPerfTest(1);
+		EndTestFailure();
 	}
 }
 
diff --git a/Engine/Plugins/Performance/AutomatedPerfTesting/Source/AutomatedPerfTesting/Private/StaticCameraTests/AutomatedStaticCameraPerfTestBase.cpp b/Engine/Plugins/Performance/AutomatedPerfTesting/Source/AutomatedPerfTesting/Private/StaticCameraTests/AutomatedStaticCameraPerfTestBase.cpp
--- a/Engine/Plugins/Performance/AutomatedPerfTesting/Source/AutomatedPerfTesting/Private/StaticCameraTests/AutomatedStaticCameraPerfTestBase.cpp
+++ b/Engine/Plugins/Performance/AutomatedPerfTesting/Source/AutomatedPerfTesting/Private/StaticCameraTests/AutomatedStaticCameraPerfTestBase.cpp
@@ -69,7 +69,7 @@ void UAutomatedStaticCameraPerfTestBase::SetupTest()
 			else
 			{
 				UE_LOG(LogAutomatedPerfTest, Error, TEXT("Invalid World when starting UAutomatedStaticCameraPerfTest, exiting..."));
-				EndAutomatedPerfTest(1);
+				EndTestFailure();
 			}
 		}
 		else
@@ -265,7 +265,7 @@ void UAutomatedStaticCameraPerfTestBase::OnInit()
 		if(MapsToTest.IsEmpty())
 		{
 			UE_LOG(LogAutomatedPerfTest, Error, TEXT("Couldn't find a map name matching %s in Static Camera Maps to Test setting. Exiting."), *CurrentMapName)
-			EndAutomatedPerfTest(1);
+			EndTestFailure();
 			return;
 		}
 	}
@@ -285,7 +285,7 @@ void UAutomatedStaticCameraPerfTestBase::OnInit()
 	if(Settings->MapsToTest.IsEmpty())
 	{
 		UE_LOG(LogAutomatedPerfTest, Error, TEXT("No maps defined in the project's Automated Perf Test | Static Camera settings. Exiting test early."));
-		EndAutomatedPerfTest(1);
+		EndTestFailure();
 		return;
 	}
 }
diff --git a/Engine/Plugins/Performance/AutomatedPerfTesting/Source/AutomatedPerfTesting/Public/AutomatedPerfTestControllerBase.h b/Engine/Plugins/Performance/AutomatedPerfTesting/Source/AutomatedPerfTesting/Public/AutomatedPerfTestControllerBase.h
--- a/Engine/Plugins/Performance/AutomatedPerfTesting/Source/AutomatedPerfTesting/Public/AutomatedPerfTestControllerBase.h
+++ b/Engine/Plugins/Performance/AutomatedPerfTesting/Source/AutomatedPerfTesting/Public/AutomatedPerfTestControllerBase.h
@@ -141,11 +141,15 @@ protected:
 	UE_API virtual void BeginDestroy() override;
 	// ~End UGauntletTestController Interface
 
+	// Shut down and queue an exit with an exit code of 0 to indicate success.
 	UFUNCTION()
-	UE_API virtual void EndAutomatedPerfTest(const int32 ExitCode = 0);
+	void EndTestSuccess();
 
-	void EndTestSuccess() { EndAutomatedPerfTest(0); }
-	void EndTestFailure() { EndAutomatedPerfTest(-1); }
+	// Shut down and queue an exit with an exit code of 1 to indicate generic failure. A custom non-zero value may be
+	// provided if needed, though any new systemic failures should be implemented as a new function with a custom 
+	// reserved value.
+	UFUNCTION()
+	void EndTestFailure(const int32 ExitCode = 1);
 
 	UFUNCTION()
 	UE_API virtual void OnVideoRecordingFinalized(bool Succeeded, const FString& FilePath);
@@ -166,6 +170,9 @@ protected:
 	/* End Profiling Functions*/
 
 private:
+	// Shut down and request the process exit with the provided error code. 
+	UE_API void EndAutomatedPerfTest(const int32 ExitCode);
+	
 	FString TraceChannels;
 	FString TestDatetime;
 	FString TestID;
