Fix Automated Performance Testing overwriting a user-provided csvMetadata command line value.

The APT Gauntlet node unconditionally set csvMetadata= to push controller-provided metadata into FCsvProfiler. That broke cases where a user also supplied csvMetadata on the command line. Both values were processed later in Gauntlet in a generic manner. Merging the lists would have required special case handling deep in generic Gauntlet property pipelines.

This changelist resolves the issue more simply by writing APT-generated metadata to a separate command line parameter reserved for APT internal use. This avoids losing user configured csvMetadata values and makes it easier to identify the provenance of metadata values.

#Synced-CL 50405523

diff --git a/Engine/Plugins/Performance/AutomatedPerfTesting/Build/Scripts/AutomatedPerfTestNode.cs b/Engine/Plugins/Performance/AutomatedPerfTesting/Build/Scripts/AutomatedPerfTestNode.cs
--- a/Engine/Plugins/Performance/AutomatedPerfTesting/Build/Scripts/AutomatedPerfTestNode.cs
+++ b/Engine/Plugins/Performance/AutomatedPerfTesting/Build/Scripts/AutomatedPerfTestNode.cs
@@ -1357,7 +1357,7 @@ namespace AutomatedPerfTest
 				CsvMetadata.Add("ReducedAsyncCompute=1");
 			}
 
-			ClientRole.CommandLineParams.Add("csvMetadata", "\"" + String.Join(",", CsvMetadata) + "\"");
+			ClientRole.CommandLineParams.Add("AutomatedPerfTest.CsvMetadata", "\"" + String.Join(",", CsvMetadata) + "\"");
 		}
 
 		private void ConfigureServerRole(TConfigClass Config)
diff --git a/Engine/Plugins/Performance/AutomatedPerfTesting/Source/AutomatedPerfTesting/Private/AutomatedPerfTestControllerBase.cpp b/Engine/Plugins/Performance/AutomatedPerfTesting/Source/AutomatedPerfTesting/Private/AutomatedPerfTestControllerBase.cpp
--- a/Engine/Plugins/Performance/AutomatedPerfTesting/Source/AutomatedPerfTesting/Private/AutomatedPerfTestControllerBase.cpp
+++ b/Engine/Plugins/Performance/AutomatedPerfTesting/Source/AutomatedPerfTesting/Private/AutomatedPerfTestControllerBase.cpp
@@ -576,7 +576,29 @@ void UAutomatedPerfTestControllerBase::OnInit()
 	{
 		UE_LOG(LogAutomatedPerfTest, Log, TEXT("Artifact Output Path: %s"), *ArtifactOutputPath)
 	}
-	
+
+#if CSV_PROFILER
+	// Translate the AutomatedPerfTest internal list of CsvMetadata provided by the Gauntlet node to CsvProfiler calls.
+	// This is provided here as well to separate systemic use by APT controllers from user use of CsvProfiler provides 
+	// csvMetadata= support.
+	FString CsvMetadataStr;
+	if (FParse::Value(FCommandLine::Get(), TEXT("AutomatedPerfTest.CsvMetadata="), CsvMetadataStr, false))
+	{
+		TArray<FString> CsvMetadataList;
+		CsvMetadataStr.ParseIntoArray(CsvMetadataList, TEXT(","), true);
+		for (int i = 0; i < CsvMetadataList.Num(); i++)
+		{
+			const FString& Metadata = CsvMetadataList[i];
+			FString Key;
+			FString Value;
+			if (Metadata.Split(TEXT("="), &Key, &Value))
+			{
+				FCsvProfiler::SetMetadata(*Key, *Value);
+			}
+		}
+	}
+#endif
+
 	FWorldDelegates::OnPreWorldInitialization.AddUObject(this, &ThisClass::OnPreWorldInitializeInternal);
 
 	InitializeInsights();
