Removing the Automated Performance Testing plugin build option DeviceProfileOverride.

APT now specifies the device profile via PerfTestConfig macro property DeviceProfile, applying the value directly through the -deviceprofile command-line flag.

Previously, the DeviceProfileOverride build parameter was passed to the APT C# generic class AutomatedPerfTestNode, which propagated the value to the client application via the command-line flag -AutomatedPerfTest.DeviceProfileOverride. This flag was parsed by AutomatedPerfTestControllerBase but was only used to override the name passed to the CSVProfiler metadata and to construct the TestID name. That functionality is preserved by directly using UDeviceProfileManager::Get().GetActiveDeviceProfileName() which will reflect the value passed in by PerfTestConfig via -deviceprofile.

#Synced-CL 50960848

#Synced-CL 50960848

diff --git a/Engine/Plugins/Performance/AutomatedPerfTesting/Build/Inc/AutomatedPerfTestCommonSettings.xml b/Engine/Plugins/Performance/AutomatedPerfTesting/Build/Inc/AutomatedPerfTestCommonSettings.xml
--- a/Engine/Plugins/Performance/AutomatedPerfTesting/Build/Inc/AutomatedPerfTestCommonSettings.xml
+++ b/Engine/Plugins/Performance/AutomatedPerfTesting/Build/Inc/AutomatedPerfTestCommonSettings.xml
@@ -19,7 +19,6 @@
 	
 	<!-- Other common arguments --> 
 	<Option Name="TestID" DefaultValue="APT" Description="A string to use to identify these performance tests" />
-	<Option Name="DeviceProfileOverride" DefaultValue="" Description="Which device profile to use for this test instead of the default value. Will be prepended by PlatformName." />
 	<Option Name="DoCSVProfiler" DefaultValue="true" Restrict="true|false" Description="Whether or not to use CSVProfiler to gather and parse performance telemetry." />
 	<Option Name="PerfReportServer" DefaultValue="false" Restrict="true|false" Description="Whether or not to upload CSVProfiler data to PRS." />
 	<Option Name="HordeTelemetry" DefaultValue="false" Restrict="true|false" Description="Whether or not to emit CSV performance summary data to Horde." />
@@ -64,7 +63,6 @@
 	<!-- Construct the common arguments for the performance test and report commands -->
 	<Property Name="AutomatedPerformanceCommonArgs" Value = "-AutomatedPerfTest.UseShippingInsights=$(UseShippingInsights)" />
 	<Property Name="AutomatedPerformanceCommonArgs" Value="$(AutomatedPerformanceCommonArgs) -AutomatedPerfTest.TestID=$(TestID)"/>
-	<Property Name="AutomatedPerformanceCommonArgs" Value="$(AutomatedPerformanceCommonArgs) -AutomatedPerfTest.DeviceProfileOverride=$(DeviceProfileOverride)" If="'$(DeviceProfileOverride)' != ''"/>
 	<Property Name="AutomatedPerformanceCommonArgs" Value="$(AutomatedPerformanceCommonArgs) -AutomatedPerfTest.TraceChannels=$(TraceChannels)" If="$(DoInsightsTrace) and '$(TraceChannels)' != ''"/>
 	<Property Name="AutomatedPerformanceCommonArgs" Value="$(AutomatedPerformanceCommonArgs) -AutomatedPerfTest.DoFPSChart" If="$(DoFPSChart)"/>
 
diff --git a/Engine/Plugins/Performance/AutomatedPerfTesting/Build/Scripts/AutomatedPerfTestConfig.cs b/Engine/Plugins/Performance/AutomatedPerfTesting/Build/Scripts/AutomatedPerfTestConfig.cs
--- a/Engine/Plugins/Performance/AutomatedPerfTesting/Build/Scripts/AutomatedPerfTestConfig.cs
+++ b/Engine/Plugins/Performance/AutomatedPerfTesting/Build/Scripts/AutomatedPerfTestConfig.cs
@@ -50,12 +50,6 @@ namespace AutomatedPerfTest
 		[AutoParamWithNames("AutomatedPerfTest.TestID")]
 		public string TestID = "";
 		
-		/// <summary>
-		/// If set, will prepend platform name and use this device profile instead of the default
-		/// </summary>
-		[AutoParamWithNames("AutomatedPerfTest.DeviceProfileOverride")]
-		public string DeviceProfileOverride = "";
-		
 		/// <summary>
 		/// If data from the test should be gathered with Unreal Insights
 		/// </summary>
diff --git a/Engine/Plugins/Performance/AutomatedPerfTesting/Build/Scripts/AutomatedPerfTestNode.cs b/Engine/Plugins/Performance/AutomatedPerfTesting/Build/Scripts/AutomatedPerfTestNode.cs
--- a/Engine/Plugins/Performance/AutomatedPerfTesting/Build/Scripts/AutomatedPerfTestNode.cs
+++ b/Engine/Plugins/Performance/AutomatedPerfTesting/Build/Scripts/AutomatedPerfTestNode.cs
@@ -1041,11 +1041,6 @@ namespace AutomatedPerfTest
 			}
 			ClientRole.CommandLineParams.AddOrAppendParamValue("AutomatedPerfTest.TestID", Config.TestID);
 
-			if (Config.DeviceProfileOverride != String.Empty)
-			{
-				ClientRole.CommandLineParams.AddOrAppendParamValue("AutomatedPerfTest.DeviceProfileOverride", Config.DeviceProfileOverride);
-			}
-
 			if (Config.DoInsightsTrace)
 			{
 				ClientRole.CommandLineParams.Add("statnamedevents");
diff --git a/Engine/Plugins/Performance/AutomatedPerfTesting/Source/AutomatedPerfTesting/Private/AutomatedPerfTestControllerBase.cpp b/Engine/Plugins/Performance/AutomatedPerfTesting/Source/AutomatedPerfTesting/Private/AutomatedPerfTestControllerBase.cpp
--- a/Engine/Plugins/Performance/AutomatedPerfTesting/Source/AutomatedPerfTesting/Private/AutomatedPerfTestControllerBase.cpp
+++ b/Engine/Plugins/Performance/AutomatedPerfTesting/Source/AutomatedPerfTesting/Private/AutomatedPerfTestControllerBase.cpp
@@ -86,17 +86,12 @@ FString UAutomatedPerfTestControllerBase::GetTestName() const
 	return TestID;
 }
 
-FString UAutomatedPerfTestControllerBase::GetDeviceProfile() const
-{
-	return DeviceProfileOverride.IsEmpty() ? UDeviceProfileManager::Get().GetActiveDeviceProfileName() : DeviceProfileOverride;  
-}
-
 FString UAutomatedPerfTestControllerBase::GetTestID() const
 {
 	const TArray<FString> TestCaseIDElements = {FApp::GetBuildVersion(),
 										  FPlatformProperties::PlatformName(),
 										  TestDatetime,
-								          GetDeviceProfile(),
+								          UDeviceProfileManager::Get().GetActiveDeviceProfileName(),
 								          GetTestName(),
 										  GetPerfTestTypeID()};
 
@@ -209,7 +204,7 @@ bool UAutomatedPerfTestControllerBase::TryStartCSVProfiler(FString CSVFileName,
 				*CSVFileName, CSVDestinationFolder.IsEmpty()? TEXT("<Default>") : *CSVDestinationFolder);
 		
 		CsvProfiler->BeginCapture(Frames, CSVDestinationFolder, CSVFileName);
-		CsvProfiler->SetDeviceProfileName(GetDeviceProfile());
+		CsvProfiler->SetDeviceProfileName(UDeviceProfileManager::Get().GetActiveDeviceProfileName());
 
 		// BeginCapture queues a capture to start on a frame boundary. It isn't possible to verify the capture
 		// is now running when this function returns.
@@ -544,7 +539,6 @@ void UAutomatedPerfTestControllerBase::OnInit()
 	// don't stop on separator because this will come in comma-separated
 	FParse::Value(FCommandLine::Get(), TEXT("AutomatedPerfTest.TraceChannels="), TraceChannels, false);
 	
-	FParse::Value(FCommandLine::Get(), TEXT("AutomatedPerfTest.DeviceProfileOverride="), DeviceProfileOverride);
 	FParse::Value(FCommandLine::Get(), TEXT("AutomatedPerfTest.TestID="), TestID);
 	
 	if (FParse::Param(FCommandLine::Get(), TEXT("AutomatedPerfTest.DoInsightsTrace")))
diff --git a/Engine/Plugins/Performance/AutomatedPerfTesting/Source/AutomatedPerfTesting/Public/AutomatedPerfTestControllerBase.h b/Engine/Plugins/Performance/AutomatedPerfTesting/Source/AutomatedPerfTesting/Public/AutomatedPerfTestControllerBase.h
--- a/Engine/Plugins/Performance/AutomatedPerfTesting/Source/AutomatedPerfTesting/Public/AutomatedPerfTestControllerBase.h
+++ b/Engine/Plugins/Performance/AutomatedPerfTesting/Source/AutomatedPerfTesting/Public/AutomatedPerfTestControllerBase.h
@@ -89,7 +89,6 @@ public:
 	UE_API UAutomatedPerfTestControllerBase(const FObjectInitializer& ObjectInitializer);
 
 	UE_API FString GetTestName() const;
-	UE_API FString GetDeviceProfile() const;
 	UE_API FString GetTestID() const;
 
 	// Unique identifier for the class type provided to output artifact generation.
@@ -176,7 +175,6 @@ private:
 	FString TraceChannels;
 	FString TestDatetime;
 	FString TestID;
-	FString DeviceProfileOverride;
 	bool bRequestsFPSChart;
 	bool bRequestsInsightsTrace;
 	bool bRequestsCSVProfiler;
