[APT] Convert trival uses of UE_LOG to UE_LOGFMT

This change transitions trivial logging statements - those that don't take any arguments - in APT from UE_LOG to UE_LOGFMT. This changelist contains only the trival instances to make it simpler to review subsequent changes that require more review.

#Synced-CL 51769453

diff --git a/Engine/Plugins/Performance/AutomatedPerfTesting/Source/AutomatedPerfTesting/Private/AutomatedMaterialPerfTest.cpp b/Engine/Plugins/Performance/AutomatedPerfTesting/Source/AutomatedPerfTesting/Private/AutomatedMaterialPerfTest.cpp
--- a/Engine/Plugins/Performance/AutomatedPerfTesting/Source/AutomatedPerfTesting/Private/AutomatedMaterialPerfTest.cpp
+++ b/Engine/Plugins/Performance/AutomatedPerfTesting/Source/AutomatedPerfTesting/Private/AutomatedMaterialPerfTest.cpp
@@ -10,6 +10,7 @@
 #include "Engine/StaticMesh.h"
 #include "Engine/Player.h"
 #include "GameFramework/Pawn.h"
+#include "Logging/StructuredLog.h"
 #include "Materials/MaterialInterface.h"
 #include "ProfilingDebugging/CsvProfiler.h"
 
@@ -35,7 +36,7 @@ void UAutomatedMaterialPerfTest::SetupTest()
 	{
 		Super::SetupTest();
 
-		UE_LOG(LogAutomatedPerfTest, Log, TEXT("UAutomatedMaterialPerfTest::SetupTest"))
+		UE_LOGFMT(LogAutomatedPerfTest, Log, "UAutomatedMaterialPerfTest::SetupTest");
 		// make sure the world exists, then create a sequence player
 		if(UWorld* const World = GetWorld())
 		{
@@ -63,7 +64,7 @@ void UAutomatedMaterialPerfTest::SetupTest()
 			Camera = GetWorld()->SpawnActor<ACameraActor>();
 			if (Camera == nullptr)
 			{
-				UE_LOG(LogAutomatedPerfTest, Error, TEXT("Failed to spawn ACameraActor. Exiting..."));
+				UE_LOGFMT(LogAutomatedPerfTest, Error, "Failed to spawn ACameraActor. Exiting...");
 
 				EndTestFailure();
 				return;
@@ -103,7 +104,7 @@ void UAutomatedMaterialPerfTest::SetupTest()
 		// if we have an invalid world, we can't run the test, so we should bail out
 		else
 		{
-			UE_LOG(LogAutomatedPerfTest, Error, TEXT("Invalid World when starting UAutomatedMaterialPerfTest, exiting..."));
+			UE_LOGFMT(LogAutomatedPerfTest, Error, "Invalid World when starting UAutomatedMaterialPerfTest, exiting...");
 			EndTestFailure();
 		}
 	}
@@ -121,7 +122,7 @@ void UAutomatedMaterialPerfTest::RunTest()
 	// The overall profiling region starts here, but each material gets its own sub region for easy delineation. 
 	MarkProfilingStart(); 
 
-	UE_LOG(LogAutomatedPerfTest, Log, TEXT("UAutomatedMaterialPerfTest::RunTest"))
+	UE_LOGFMT(LogAutomatedPerfTest, Log, "UAutomatedMaterialPerfTest::RunTest");
 	
 	FTimerHandle UnusedHandle;
 	GetWorld()->GetTimerManager().SetTimer(UnusedHandle, this, &UAutomatedMaterialPerfTest::SetUpNextMaterial, 1.0, false, Settings->WarmUpTime);
@@ -145,7 +146,7 @@ void UAutomatedMaterialPerfTest::SetUpNextMaterial()
 	
 	if(CurrentMaterialIndex >= Settings->MaterialsToTest.Num())
 	{
-		UE_LOG(LogAutomatedPerfTest, Log, TEXT("No more materials left to test, moving to teardown."))
+		UE_LOGFMT(LogAutomatedPerfTest, Log, "No more materials left to test, moving to teardown.");
 		TeardownTest();
 		return;
 	}
@@ -244,14 +245,14 @@ void UAutomatedMaterialPerfTest::OnInit()
 {
 	Super::OnInit();
 	
-	UE_LOG(LogAutomatedPerfTest, Log, TEXT("UAutomatedMaterialPerfTest::OnInit"));
+	UE_LOGFMT(LogAutomatedPerfTest, Log, "UAutomatedMaterialPerfTest::OnInit");
 
 	Settings = GetDefault<UAutomatedMaterialPerfTestProjectSettings>();
 
 	// early out if there aren't actually any materials
 	if (Settings->MaterialsToTest.IsEmpty())
 	{
-		UE_LOG(LogAutomatedPerfTest, Error, TEXT("No materials defined in the project's Automated Perf Test | Materials settings. Exiting test early."));
+		UE_LOGFMT(LogAutomatedPerfTest, Error, "No materials defined in the project's Automated Perf Test | Materials settings. Exiting test early.");
 		EndTestFailure();
 	}
 }
@@ -268,7 +269,7 @@ void UAutomatedMaterialPerfTest::UnbindAllDelegates()
 
 void UAutomatedMaterialPerfTest::TeardownTest(bool bExitAfterTeardown)
 {
-	UE_LOG(LogAutomatedPerfTest, Log, TEXT("AutomatedMaterialPerfTest::TeardownTest"));
+	UE_LOGFMT(LogAutomatedPerfTest, Log, "AutomatedMaterialPerfTest::TeardownTest");
 	MarkProfilingEnd();
 	TeardownProfiling();
 	Super::TeardownTest(bExitAfterTeardown);
diff --git a/Engine/Plugins/Performance/AutomatedPerfTesting/Source/AutomatedPerfTesting/Private/AutomatedPerfTestControllerBase.cpp b/Engine/Plugins/Performance/AutomatedPerfTesting/Source/AutomatedPerfTesting/Private/AutomatedPerfTestControllerBase.cpp
--- a/Engine/Plugins/Performance/AutomatedPerfTesting/Source/AutomatedPerfTesting/Private/AutomatedPerfTestControllerBase.cpp
+++ b/Engine/Plugins/Performance/AutomatedPerfTesting/Source/AutomatedPerfTesting/Private/AutomatedPerfTestControllerBase.cpp
@@ -22,6 +22,7 @@
 #include "UnrealClient.h"
 #include "UnrealEngine.h"
 #include "GauntletModule.h"
+#include "Logging/StructuredLog.h"
 
 #include UE_INLINE_GENERATED_CPP_BY_NAME(AutomatedPerfTestControllerBase)
 
@@ -62,7 +63,7 @@ void UAutomatedPerfTestControllerBase::OnPreWorldInitialize(UWorld* World)
 
 void UAutomatedPerfTestControllerBase::OnWorldBeginPlay()
 {
-	UE_LOG(LogAutomatedPerfTest, Log, TEXT("OnWorldBeginPlay"));
+	UE_LOGFMT(LogAutomatedPerfTest, Log, "OnWorldBeginPlay");
 	SetupTest();
 }
 
@@ -211,7 +212,7 @@ bool UAutomatedPerfTestControllerBase::TryStartCSVProfiler(FString CSVFileName,
 		return true;
 	}
 #endif
-	UE_LOG(LogAutomatedPerfTest, Warning, TEXT("CSVProfiler Start requested, but not available."))
+	UE_LOGFMT(LogAutomatedPerfTest, Warning, "CSVProfiler Start requested, but not available.");
 	return false;
 }
 
@@ -231,7 +232,7 @@ bool UAutomatedPerfTestControllerBase::TryStopCSVProfiler()
 		return true;
 	}
 #endif
-	UE_LOG(LogAutomatedPerfTest, Warning, TEXT("CSVProfiler Stop requested, but not available."))
+	UE_LOGFMT(LogAutomatedPerfTest, Warning, "CSVProfiler Stop requested, but not available.");
 	return false;
 }
 
@@ -275,17 +276,17 @@ bool UAutomatedPerfTestControllerBase::TryStartVideoCapture()
 			}
 			else
 			{
-				UE_LOG(LogAutomatedPerfTest, Warning, TEXT("Video recording could not be enabled."));
+				UE_LOGFMT(LogAutomatedPerfTest, Warning, "Video recording could not be enabled.");
 			}
 		}
 		else
 		{
-			UE_LOG(LogAutomatedPerfTest, Warning, TEXT("Could not start a new recording, may be already recording."));
+			UE_LOGFMT(LogAutomatedPerfTest, Warning, "Could not start a new recording, may be already recording.");
 		}
 	}
 	else
 	{
-		UE_LOG(LogAutomatedPerfTest, Warning, TEXT("Video recording system is null."));
+		UE_LOGFMT(LogAutomatedPerfTest, Warning, "Video recording system is null.");
 	}
 
 	return false;	
@@ -301,7 +302,7 @@ bool UAutomatedPerfTestControllerBase::TryFinalizingVideoCapture(const bool bSto
 
 			if (VideoRecordingSystem->GetRecordingState() == EVideoRecordingState::Finalizing)
 			{
-				UE_LOG(LogAutomatedPerfTest, Log, TEXT("Finalizing recording..."));
+				UE_LOGFMT(LogAutomatedPerfTest, Log, "Finalizing recording...");
 				VideoRecordingSystem->GetOnVideoRecordingFinalizedDelegate().AddUObject(this, &ThisClass::OnVideoRecordingFinalized);
 				return true;
 			}
@@ -317,7 +318,7 @@ bool UAutomatedPerfTestControllerBase::TryFinalizingVideoCapture(const bool bSto
 	}
 	else
 	{
-		UE_LOG(LogAutomatedPerfTest, Warning, TEXT("Video recording system is null."));
+		UE_LOGFMT(LogAutomatedPerfTest, Warning, "Video recording system is null.");
 	}
 
 	return false;	
@@ -325,7 +326,7 @@ bool UAutomatedPerfTestControllerBase::TryFinalizingVideoCapture(const bool bSto
 
 void UAutomatedPerfTestControllerBase::SetupTest()
 {
-	UE_LOG(LogAutomatedPerfTest, Log, TEXT("Base:: SetupTest"));
+	UE_LOGFMT(LogAutomatedPerfTest, Log, "Base:: SetupTest");
 
 	SetupGameModeInstance();
 
@@ -335,7 +336,7 @@ void UAutomatedPerfTestControllerBase::SetupTest()
 
 void UAutomatedPerfTestControllerBase::RunTest()
 {
-	UE_LOG(LogAutomatedPerfTest, Log, TEXT("Base:: RunTest"));
+	UE_LOGFMT(LogAutomatedPerfTest, Log, "Base:: RunTest");
 
 	if(GameMode && GameMode->GetClass()->ImplementsInterface(UAutomatedPerfTestInterface::StaticClass()))
 	{
@@ -347,7 +348,7 @@ void UAutomatedPerfTestControllerBase::RunTest()
 
 void UAutomatedPerfTestControllerBase::TeardownTest(bool bExitAfterTeardown)
 {
-	UE_LOG(LogAutomatedPerfTest, Log, TEXT("Base:: TeardownTest"));
+	UE_LOGFMT(LogAutomatedPerfTest, Log, "Base:: TeardownTest");
 
 	if(GameMode && GameMode->GetClass()->ImplementsInterface(UAutomatedPerfTestInterface::StaticClass()))
 	{
@@ -514,7 +515,7 @@ void UAutomatedPerfTestControllerBase::TeardownProfiling()
 	{
 		if (!TryFinalizingVideoCapture())
 		{
-			UE_LOG(LogAutomatedPerfTest, Warning, TEXT("Attempted to finalize requested video capture, but failed."))
+			UE_LOGFMT(LogAutomatedPerfTest, Warning, "Attempted to finalize requested video capture, but failed.");
 		}
 	}
 }
@@ -534,7 +535,7 @@ void UAutomatedPerfTestControllerBase::OnInit()
 {
 	Super::OnInit();
 
-	UE_LOG(LogAutomatedPerfTest, Log, TEXT("Base:: OnInit"));
+	UE_LOGFMT(LogAutomatedPerfTest, Log, "Base:: OnInit");
 
 	// don't stop on separator because this will come in comma-separated
 	FParse::Value(FCommandLine::Get(), TEXT("AutomatedPerfTest.TraceChannels="), TraceChannels, false);
@@ -543,27 +544,27 @@ void UAutomatedPerfTestControllerBase::OnInit()
 	
 	if (FParse::Param(FCommandLine::Get(), TEXT("AutomatedPerfTest.DoInsightsTrace")))
 	{
-		UE_LOG(LogAutomatedPerfTest, Log, TEXT("Insights Trace Requested"))
+		UE_LOGFMT(LogAutomatedPerfTest, Log, "Insights Trace Requested");
 		bRequestsInsightsTrace = true;
 	}
 	if (FParse::Param(FCommandLine::Get(), TEXT("AutomatedPerfTest.DoCSVProfiler")))
 	{
-		UE_LOG(LogAutomatedPerfTest, Log, TEXT("CSV Profiler Requested"))
+		UE_LOGFMT(LogAutomatedPerfTest, Log, "CSV Profiler Requested");
 		bRequestsCSVProfiler = true;
 	}
 	if (FParse::Param(FCommandLine::Get(), TEXT("AutomatedPerfTest.DoFPSChart")))
 	{
-		UE_LOG(LogAutomatedPerfTest, Log, TEXT("FPSCharts Requested"))
+		UE_LOGFMT(LogAutomatedPerfTest, Log, "FPSCharts Requested");
 		bRequestsFPSChart = true;
 	}
 	if (FParse::Param(FCommandLine::Get(), TEXT("AutomatedPerfTest.DoVideoCapture")))
 	{
-		UE_LOG(LogAutomatedPerfTest, Log, TEXT("Video Capture Requested"))
+		UE_LOGFMT(LogAutomatedPerfTest, Log, "Video Capture Requested");
 		bRequestsVideoCapture = true;
 	}
 	if (FParse::Param(FCommandLine::Get(), TEXT("AutomatedPerfTest.LockDynamicRes")))
 	{
-		UE_LOG(LogAutomatedPerfTest, Log, TEXT("Locking dynamic res requested"))
+		UE_LOGFMT(LogAutomatedPerfTest, Log, "Locking dynamic res requested");
 		bRequestsLockedDynRes = true;
 	}
 	if (FParse::Value(FCommandLine::Get(), TEXT("AutomatedPerfTest.ArtifactOutputPath="), ArtifactOutputPath))
@@ -638,7 +639,7 @@ void UAutomatedPerfTestControllerBase::OnVideoRecordingFinalized(bool Succeeded,
 {
 	if(!Succeeded)
 	{
-		UE_LOG(LogAutomatedPerfTest, Warning, TEXT("Video Capture finalized, but did not succeed"))
+		UE_LOGFMT(LogAutomatedPerfTest, Warning, "Video Capture finalized, but did not succeed");
 	}
 /* TODO moving this to automation layer due to file access restrictions
 	FString SrcFilePath;
@@ -655,7 +656,7 @@ void UAutomatedPerfTestControllerBase::OnVideoRecordingFinalized(bool Succeeded,
 	
 	if(IFileManager::Get().Copy(*DestinationFilePath, *FilePath, 1, 1) != COPY_OK)
 	{
-		UE_LOG(LogAutomatedPerfTest, Warning, TEXT("Failed to copy video file"));
+		UE_LOGFMT(LogAutomatedPerfTest, Warning, "Failed to copy video file");
 	}
 	*/
 }
diff --git a/Engine/Plugins/Performance/AutomatedPerfTesting/Source/AutomatedPerfTesting/Private/AutomatedPerfTestGameModeBase.cpp b/Engine/Plugins/Performance/AutomatedPerfTesting/Source/AutomatedPerfTesting/Private/AutomatedPerfTestGameModeBase.cpp
--- a/Engine/Plugins/Performance/AutomatedPerfTesting/Source/AutomatedPerfTesting/Private/AutomatedPerfTestGameModeBase.cpp
+++ b/Engine/Plugins/Performance/AutomatedPerfTesting/Source/AutomatedPerfTesting/Private/AutomatedPerfTestGameModeBase.cpp
@@ -2,23 +2,24 @@
 
 #include "AutomatedPerfTestGameModeBase.h"
 #include "AutomatedPerfTesting.h"
+#include "Logging/StructuredLog.h"
 
 void AAutomatedPerfTestGameModeBase::SetupTest_Implementation()
 {
-	UE_LOG(LogAutomatedPerfTest, Log, TEXT("AutomatedPerfTestGameModeBase SetupTest"))
+	UE_LOGFMT(LogAutomatedPerfTest, Log, "AutomatedPerfTestGameModeBase SetupTest");
 }
 
 void AAutomatedPerfTestGameModeBase::TeardownTest_Implementation()
 {
-	UE_LOG(LogAutomatedPerfTest, Log, TEXT("AutomatedPerfTestGameModeBase TeardownTest"))
+	UE_LOGFMT(LogAutomatedPerfTest, Log, "AutomatedPerfTestGameModeBase TeardownTest");
 }
 
 void AAutomatedPerfTestGameModeBase::RunTest_Implementation()
 {
-	UE_LOG(LogAutomatedPerfTest, Log, TEXT("AutomatedPerfTestGameModeBase RunTest"))
+	UE_LOGFMT(LogAutomatedPerfTest, Log, "AutomatedPerfTestGameModeBase RunTest");
 }
 
 void AAutomatedPerfTestGameModeBase::Exit_Implementation()
 {
-	UE_LOG(LogAutomatedPerfTest, Log, TEXT("AutomatedPerfTestGameModeBase Exit"))
+	UE_LOGFMT(LogAutomatedPerfTest, Log, "AutomatedPerfTestGameModeBase Exit");
 }
\ No newline at end of file
diff --git a/Engine/Plugins/Performance/AutomatedPerfTesting/Source/AutomatedPerfTesting/Private/AutomatedReplayPerfTest.cpp b/Engine/Plugins/Performance/AutomatedPerfTesting/Source/AutomatedPerfTesting/Private/AutomatedReplayPerfTest.cpp
--- a/Engine/Plugins/Performance/AutomatedPerfTesting/Source/AutomatedPerfTesting/Private/AutomatedReplayPerfTest.cpp
+++ b/Engine/Plugins/Performance/AutomatedPerfTesting/Source/AutomatedPerfTesting/Private/AutomatedReplayPerfTest.cpp
@@ -5,6 +5,7 @@
 #include "Engine/World.h"
 #include "Engine/Engine.h"
 #include "Engine/GameInstance.h"
+#include "Logging/StructuredLog.h"
 #include "Misc/CommandLine.h"
 #include "Misc/Paths.h"
 #include "Net/UnrealNetwork.h"
@@ -76,12 +77,12 @@ void UAutomatedReplayPerfTest::SetupTest()
 		FNetworkReplayDelegates::OnReplayPlaybackComplete.AddUObject(this, &ThisClass::OnReplayComplete);
 		FNetworkReplayDelegates::OnReplayPlaybackFailure.AddUObject(this, &ThisClass::OnReplayFailure);
 
-		UE_LOG(LogAutomatedPerfTest, Log, TEXT("Starting Replay Perf Test"));
+		UE_LOGFMT(LogAutomatedPerfTest, Log, "Starting Replay Perf Test");
 		RunTest();
 	}
 	else
 	{
-		UE_LOG(LogAutomatedPerfTest, Error, TEXT("Invalid World when starting AutomatedReplayPerfTest, exiting..."));
+		UE_LOGFMT(LogAutomatedPerfTest, Error, "Invalid World when starting AutomatedReplayPerfTest, exiting...");
 		EndTestFailure();
 	}
 }
@@ -100,7 +101,7 @@ void UAutomatedReplayPerfTest::RunTest()
 
 	if (!bIsReplayTriggered)
 	{
-		UE_LOG(LogAutomatedPerfTest, Error, TEXT("Could not start Replay Perf Test"));
+		UE_LOGFMT(LogAutomatedPerfTest, Error, "Could not start Replay Perf Test");
 		EndTestFailure();
 	}
 }
@@ -114,7 +115,7 @@ void UAutomatedReplayPerfTest::TeardownTest(bool bExitAfterTeardown)
 		return;
 	}
 
-	UE_LOG(LogAutomatedPerfTest, Log, TEXT("AutomatedReplayPerfTest::TeardownTest"));
+	UE_LOGFMT(LogAutomatedPerfTest, Log, "AutomatedReplayPerfTest::TeardownTest");
 
 	bIsTearingDown = true;
 	MarkProfilingEnd();
@@ -132,7 +133,7 @@ void UAutomatedReplayPerfTest::OnInit()
 {
 	Super::OnInit();
 
-	UE_LOG(LogAutomatedPerfTest, Log, TEXT("AutomatedReplayPerfTest::OnInit"));
+	UE_LOGFMT(LogAutomatedPerfTest, Log, "AutomatedReplayPerfTest::OnInit");
 
 	// Replay subsystem looks for replay files in this folder by default.
 	const FString DefaultReplayPath = FPaths::Combine(*FPaths::ProjectSavedDir(), TEXT("Demos/")); 
@@ -185,7 +186,7 @@ void UAutomatedReplayPerfTest::OnInit()
 
 	if (ReplayName.IsEmpty() || !bReplayExists)
 	{
-		UE_LOG(LogAutomatedPerfTest, Error, TEXT("Replay not specified in args nor found in settings"));
+		UE_LOGFMT(LogAutomatedPerfTest, Error, "Replay not specified in args nor found in settings");
 		EndTestFailure();
 		return;
 	}
diff --git a/Engine/Plugins/Performance/AutomatedPerfTesting/Source/AutomatedPerfTesting/Private/AutomatedSequencePerfTest.cpp b/Engine/Plugins/Performance/AutomatedPerfTesting/Source/AutomatedPerfTesting/Private/AutomatedSequencePerfTest.cpp
--- a/Engine/Plugins/Performance/AutomatedPerfTesting/Source/AutomatedPerfTesting/Private/AutomatedSequencePerfTest.cpp
+++ b/Engine/Plugins/Performance/AutomatedPerfTesting/Source/AutomatedPerfTesting/Private/AutomatedSequencePerfTest.cpp
@@ -10,6 +10,7 @@
 #include "Engine/LevelStreaming.h"
 #include "MovieSceneSequencePlayer.h"
 #include "Kismet/GameplayStatics.h"
+#include "Logging/StructuredLog.h"
 #include "Misc/CommandLine.h"
 #include "ProfilingDebugging/CsvProfiler.h"
 
@@ -67,12 +68,12 @@ void UAutomatedSequencePerfTest::SetupTest()
 				ULevelSequence* TargetSequence = LoadObject<ULevelSequence>(NULL, *CurrentMapSequenceCombo->Sequence.ToString(), NULL, LOAD_None, NULL);
 				check(TargetSequence);
 		
-				UE_LOG(LogAutomatedPerfTest, Log, TEXT("World is valid, creating sequence player"));
+				UE_LOGFMT(LogAutomatedPerfTest, Log, "World is valid, creating sequence player");
 				SequencePlayer = ULevelSequencePlayer::CreateLevelSequencePlayer(World, TargetSequence, FMovieSceneSequencePlaybackSettings(), SequenceActor);
 
 				if (SequencePlayer == nullptr)
 				{
-					UE_LOG(LogAutomatedPerfTest, Error, TEXT("Unable to create sequence player when starting AutomatedSequencePerfTest, exiting..."));
+					UE_LOGFMT(LogAutomatedPerfTest, Error, "Unable to create sequence player when starting AutomatedSequencePerfTest, exiting...");
 					EndTestFailure();
 					return;
 				}
@@ -84,7 +85,7 @@ void UAutomatedSequencePerfTest::SetupTest()
 				PlaybackParams.Time = 0.0;
 				PlaybackParams.UpdateMethod = EUpdatePositionMethod::Scrub;
 
-				UE_LOG(LogAutomatedPerfTest, Log, TEXT("SetupMapTest:: Scrubbing to start"));
+				UE_LOGFMT(LogAutomatedPerfTest, Log, "SetupMapTest:: Scrubbing to start");
 				SequencePlayer->PlayTo(PlaybackParams, PlayToParams);
 
 				// Setup profiling only when we are sure to run the sequence. 
@@ -97,7 +98,7 @@ void UAutomatedSequencePerfTest::SetupTest()
 			// if we have an invalid world, we can't run the test, so we should bail out
 			else
 			{
-				UE_LOG(LogAutomatedPerfTest, Error, TEXT("Invalid World when starting AutomatedSequencePerfTest, exiting..."));
+				UE_LOGFMT(LogAutomatedPerfTest, Error, "Invalid World when starting AutomatedSequencePerfTest, exiting...");
 				EndTestFailure();
 				return;
 			}
@@ -110,14 +111,14 @@ void UAutomatedSequencePerfTest::SetupTest()
 	}
 	else
 	{
-		UE_LOG(LogAutomatedPerfTest, Log, TEXT("Current Map Sequence Combo has not been set, calling NextMap to trigger."))
+		UE_LOGFMT(LogAutomatedPerfTest, Log, "Current Map Sequence Combo has not been set, calling NextMap to trigger.");
 		NextMap();
 	}
 }
 
 void UAutomatedSequencePerfTest::NextMap()
 {
-	UE_LOG(LogAutomatedPerfTest, Log, TEXT("UAutomatedSequencePerfTest::NextMap"))
+	UE_LOGFMT(LogAutomatedPerfTest, Log, "UAutomatedSequencePerfTest::NextMap");
 
 	if(MapSequenceCombos.Num() > 0)
 	{
@@ -146,7 +147,7 @@ void UAutomatedSequencePerfTest::NextMap()
 	}
 	else
 	{
-		UE_LOG(LogAutomatedPerfTest, Log, TEXT("UAutomatedSequencePerfTest::NextMap, all maps complete, exiting after delay."))
+		UE_LOGFMT(LogAutomatedPerfTest, Log, "UAutomatedSequencePerfTest::NextMap, all maps complete, exiting after delay.");
 		TriggerExitAfterDelay();
 	}
 }
@@ -155,12 +156,12 @@ void UAutomatedSequencePerfTest::RunTest()
 {
 	Super::RunTest();
 
-	UE_LOG(LogAutomatedPerfTest, Log, TEXT("RunTest"));
+	UE_LOGFMT(LogAutomatedPerfTest, Log, "RunTest");
 	
 	// make sure we have a valid sequence player
 	if(SequencePlayer)
 	{
-		UE_LOG(LogAutomatedPerfTest, Log, TEXT("RunTest::Valid Sequence Player, proceeding"));
+		UE_LOGFMT(LogAutomatedPerfTest, Log, "RunTest::Valid Sequence Player, proceeding");
 		MarkProfilingStart(); // Mark profiler start region on sequence start. 
 
 		if (RequestsCSVProfiler())
@@ -173,7 +174,7 @@ void UAutomatedSequencePerfTest::RunTest()
 			else
 			{
 				CSVFilename = GetTestID();
-				UE_LOG(LogAutomatedPerfTest, Error, TEXT("Current Map Sequence Combo Name not set"));
+				UE_LOGFMT(LogAutomatedPerfTest, Error, "Current Map Sequence Combo Name not set");
 			}
 
 			TryStartCSVProfiler(CSVFilename);
@@ -192,7 +193,7 @@ void UAutomatedSequencePerfTest::RunTest()
 	// otherwise bail out of the test
 	else
 	{
-		UE_LOG(LogAutomatedPerfTest, Error, TEXT("Invalid SequencePlayer when starting AutomatedSequencePerfTest, exiting..."));
+		UE_LOGFMT(LogAutomatedPerfTest, Error, "Invalid SequencePlayer when starting AutomatedSequencePerfTest, exiting...");
 		EndTestFailure();
 	}
 }
@@ -207,7 +208,7 @@ void UAutomatedSequencePerfTest::OnSequenceFinished()
 
 void UAutomatedSequencePerfTest::TeardownTest(bool bExitAfterTeardown)
 {
-	UE_LOG(LogAutomatedPerfTest, Log, TEXT("AutomatedSequencePerfTest::TeardownTest"));
+	UE_LOGFMT(LogAutomatedPerfTest, Log, "AutomatedSequencePerfTest::TeardownTest");
 	
 	MarkProfilingEnd();
 
@@ -309,7 +310,7 @@ FString UAutomatedSequencePerfTest::GetCameraCutID()
 void UAutomatedSequencePerfTest::OnInit()
 {
 	Super::OnInit();
-	UE_LOG(LogAutomatedPerfTest, Log, TEXT("AutomatedSequencePerfTest::OnInit"));
+	UE_LOGFMT(LogAutomatedPerfTest, Log, "AutomatedSequencePerfTest::OnInit");
 	
 	Settings = GetDefault<UAutomatedSequencePerfTestProjectSettings>();
 
diff --git a/Engine/Plugins/Performance/AutomatedPerfTesting/Source/AutomatedPerfTesting/Private/ProfileGo/ProfileGo.cpp b/Engine/Plugins/Performance/AutomatedPerfTesting/Source/AutomatedPerfTesting/Private/ProfileGo/ProfileGo.cpp
--- a/Engine/Plugins/Performance/AutomatedPerfTesting/Source/AutomatedPerfTesting/Private/ProfileGo/ProfileGo.cpp
+++ b/Engine/Plugins/Performance/AutomatedPerfTesting/Source/AutomatedPerfTesting/Private/ProfileGo/ProfileGo.cpp
@@ -143,7 +143,7 @@ bool UProfileGo::Run(UProfileGoSubsystem* InSubsystem, const TCHAR* InProfileNam
 
 	if (!InSubsystem)
 	{
-		UE_LOG(LogAutomatedPerfTest, Log, TEXT("InSubsystem is nullptr! Aborting ProfileGo."));
+		UE_LOGFMT(LogAutomatedPerfTest, Log, "InSubsystem is nullptr! Aborting ProfileGo.");
 		return false;
 	}
 
@@ -417,14 +417,14 @@ bool UProfileGo::WaitForLoadingAndStreaming(UWorld* InWorld, float DeltaTime, bo
 	if (IsAsyncLoading())
 	{
 		bLoadInProgress = true;
-		UE_LOG(LogAutomatedPerfTest, Log, TEXT("TickSettling: Waiting On AsyncLoading"));
+		UE_LOGFMT(LogAutomatedPerfTest, Log, "TickSettling: Waiting On AsyncLoading");
 	}
 
 	TWeakObjectPtr<UWorldPartitionSubsystem> WorldPartition = UWorld::GetSubsystem<UWorldPartitionSubsystem>(InWorld);
 	// Ensure world partition streaming completes. Should be redundant with above IsStreamingInLevels()
 	if (WorldPartition.Get() && !WorldPartition->IsAllStreamingCompleted())
 	{
-		UE_LOG(LogAutomatedPerfTest, Log, TEXT("TickSettling: Waiting On World Partition Level Streaming"));
+		UE_LOGFMT(LogAutomatedPerfTest, Log, "TickSettling: Waiting On World Partition Level Streaming");
 		bLoadInProgress = true;
 	}
 
@@ -432,7 +432,7 @@ bool UProfileGo::WaitForLoadingAndStreaming(UWorld* InWorld, float DeltaTime, bo
 	if (GShaderCompilingManager && GShaderCompilingManager->GetNumOutstandingJobs())
 	{
 		GShaderCompilingManager->FinishAllCompilation();
-		UE_LOG(LogAutomatedPerfTest, Log, TEXT("TickSettling: Waiting On Shaders"));
+		UE_LOGFMT(LogAutomatedPerfTest, Log, "TickSettling: Waiting On Shaders");
 		bLoadInProgress = true;
 	}
 
@@ -442,7 +442,7 @@ bool UProfileGo::WaitForLoadingAndStreaming(UWorld* InWorld, float DeltaTime, bo
 		ResourceStreamingIteration++;
 		if (!bResourceStreamingCompleted)
 		{
-			UE_LOG(LogAutomatedPerfTest, Log, TEXT("TickSettling: Waiting On Resource Streaming Final update"));
+			UE_LOGFMT(LogAutomatedPerfTest, Log, "TickSettling: Waiting On Resource Streaming Final update");
 			bLoadInProgress = true;
 		}
 	}
@@ -1247,7 +1247,7 @@ void UProfileGo::RegisterInternalCommandDelegates()
 
 		if (CsvProfiler->IsCapturing())
 		{
-			UE_LOG(LogAutomatedPerfTest, Warning, TEXT("Tried to begin csvprofile when csvprofiler was already capturing!"));
+			UE_LOGFMT(LogAutomatedPerfTest, Warning, "Tried to begin csvprofile when csvprofiler was already capturing!");
 		}
 		else
 		{
@@ -1326,7 +1326,7 @@ void UProfileGo::RegisterInternalCommandDelegates()
 #if CSV_PROFILER
 		if (!FCsvProfiler::Get()->IsCapturing())
 		{
-			UE_LOG(LogAutomatedPerfTest, Warning, TEXT("Tried to end csvprofile when csvprofiler was not capturing!"));
+			UE_LOGFMT(LogAutomatedPerfTest, Warning, "Tried to end csvprofile when csvprofiler was not capturing!");
 		}
 		else
 		{
@@ -1846,7 +1846,7 @@ bool UProfileGo::FindBestStandingPoint(UWorld* InWorld, FVector StartXyHalfHeigh
 
 	if (!InWorld)
 	{
-		UE_LOG(LogAutomatedPerfTest, Log, TEXT("InWorld is nullptr! Cannot trace for best standing point."));
+		UE_LOGFMT(LogAutomatedPerfTest, Log, "InWorld is nullptr! Cannot trace for best standing point.");
 		return false;
 	}
 
diff --git a/Engine/Plugins/Performance/AutomatedPerfTesting/Source/AutomatedPerfTesting/Private/StaticCameraTests/AutomatedStaticCameraPerfTestBase.cpp b/Engine/Plugins/Performance/AutomatedPerfTesting/Source/AutomatedPerfTesting/Private/StaticCameraTests/AutomatedStaticCameraPerfTestBase.cpp
--- a/Engine/Plugins/Performance/AutomatedPerfTesting/Source/AutomatedPerfTesting/Private/StaticCameraTests/AutomatedStaticCameraPerfTestBase.cpp
+++ b/Engine/Plugins/Performance/AutomatedPerfTesting/Source/AutomatedPerfTesting/Private/StaticCameraTests/AutomatedStaticCameraPerfTestBase.cpp
@@ -5,6 +5,7 @@
 #include "Kismet/GameplayStatics.h"
 #include "AutomatedPerfTesting.h"
 #include "Camera/CameraActor.h"
+#include "Logging/StructuredLog.h"
 #include "Misc/CommandLine.h"
 #include "ProfilingDebugging/CsvProfiler.h"
 
@@ -41,7 +42,7 @@ void UAutomatedStaticCameraPerfTestBase::SetupTest()
 		if(GetCurrentMap() == CurrentMapPath.GetAssetName())
 		{
 			Super::SetupTest();
-			UE_LOG(LogAutomatedPerfTest, Log, TEXT("UAutomatedStaticCameraPerfTestBase::SetupTest"))
+			UE_LOGFMT(LogAutomatedPerfTest, Log, "UAutomatedStaticCameraPerfTestBase::SetupTest");
 			// make sure the world exists, then create a sequence player
 			if(UWorld* const World = GetWorld())
 			{
@@ -68,7 +69,7 @@ void UAutomatedStaticCameraPerfTestBase::SetupTest()
 			// if we have an invalid world, we can't run the test, so we should bail out
 			else
 			{
-				UE_LOG(LogAutomatedPerfTest, Error, TEXT("Invalid World when starting UAutomatedStaticCameraPerfTest, exiting..."));
+				UE_LOGFMT(LogAutomatedPerfTest, Error, "Invalid World when starting UAutomatedStaticCameraPerfTest, exiting...");
 				EndTestFailure();
 			}
 		}
@@ -80,7 +81,7 @@ void UAutomatedStaticCameraPerfTestBase::SetupTest()
 	}
 	else
 	{
-		UE_LOG(LogAutomatedPerfTest, Log, TEXT("Current Map Path has not been set, calling NextMap to trigger."))
+		UE_LOGFMT(LogAutomatedPerfTest, Log, "Current Map Path has not been set, calling NextMap to trigger.");
 		NextMap();
 	}
 }
@@ -89,7 +90,7 @@ void UAutomatedStaticCameraPerfTestBase::RunTest()
 {
 	Super::RunTest();
 
-	UE_LOG(LogAutomatedPerfTest, Log, TEXT("UAutomatedStaticCameraPerfTestBase::RunTest"))
+	UE_LOGFMT(LogAutomatedPerfTest, Log, "UAutomatedStaticCameraPerfTestBase::RunTest");
 	MarkProfilingStart(); // Mark the start region on the profiler once the test starts
 	
 	FTimerHandle UnusedHandle;
@@ -113,7 +114,7 @@ void UAutomatedStaticCameraPerfTestBase::SetUpNextCamera()
 {
 	if(CamerasToTest.Num() <= 0)
 	{
-		UE_LOG(LogAutomatedPerfTest, Log, TEXT("No more cameras left to test, moving to next map."))
+		UE_LOGFMT(LogAutomatedPerfTest, Log, "No more cameras left to test, moving to next map.");
 		NextMap();
 		return;
 	}
@@ -166,7 +167,7 @@ void UAutomatedStaticCameraPerfTestBase::ScreenshotCamera()
 
 void UAutomatedStaticCameraPerfTestBase::NextMap()
 {
-	UE_LOG(LogAutomatedPerfTest, Log, TEXT("UAutomatedStaticCameraPerfTestBase::NextMap"))
+	UE_LOGFMT(LogAutomatedPerfTest, Log, "UAutomatedStaticCameraPerfTestBase::NextMap");
 
 	if(MapsToTest.Num() > 0)
 	{
@@ -186,14 +187,14 @@ void UAutomatedStaticCameraPerfTestBase::NextMap()
 	}
 	else
 	{
-		UE_LOG(LogAutomatedPerfTest, Log, TEXT("UAutomatedStaticCameraPerfTestBase::NextMap, all maps complete, exiting after delay."))
+		UE_LOGFMT(LogAutomatedPerfTest, Log, "UAutomatedStaticCameraPerfTestBase::NextMap, all maps complete, exiting after delay.");
 		TriggerExitAfterDelay();
 	}
 }
 
 TArray<ACameraActor*> UAutomatedStaticCameraPerfTestBase::GetMapCameraActors()
 {
-	UE_LOG(LogAutomatedPerfTest, Warning, TEXT("GetMapCameraActors called in base class UAutomatedStaticCameraPerfTestBase, please ensure you've overridden this function in the subclass, and you're not using the base class as your test controller"));
+	UE_LOGFMT(LogAutomatedPerfTest, Warning, "GetMapCameraActors called in base class UAutomatedStaticCameraPerfTestBase, please ensure you've overridden this function in the subclass, and you're not using the base class as your test controller");
 	return TArray<ACameraActor*>();
 }
 
@@ -248,7 +249,7 @@ void UAutomatedStaticCameraPerfTestBase::OnInit()
 {
 	Super::OnInit();
 	
-	UE_LOG(LogAutomatedPerfTest, Log, TEXT("UAutomatedStaticCameraPerfTestBase::OnInit"));
+	UE_LOGFMT(LogAutomatedPerfTest, Log, "UAutomatedStaticCameraPerfTestBase::OnInit");
 
 	Settings = GetDefault<UAutomatedStaticCameraPerfTestProjectSettings>();
 
@@ -272,7 +273,7 @@ void UAutomatedStaticCameraPerfTestBase::OnInit()
 	// otherwise, use all the maps defined in project settings
 	else
 	{
-		UE_LOG(LogAutomatedPerfTest, Log, TEXT("No map name specified, testing all maps."))
+		UE_LOGFMT(LogAutomatedPerfTest, Log, "No map name specified, testing all maps.");
 		for(FSoftObjectPath MapName : Settings->MapsToTest)
 		{
 			MapsToTest.Add(MapName);
@@ -284,7 +285,7 @@ void UAutomatedStaticCameraPerfTestBase::OnInit()
 	// early out if there aren't actually any maps set in project settings
 	if(Settings->MapsToTest.IsEmpty())
 	{
-		UE_LOG(LogAutomatedPerfTest, Error, TEXT("No maps defined in the project's Automated Perf Test | Static Camera settings. Exiting test early."));
+		UE_LOGFMT(LogAutomatedPerfTest, Error, "No maps defined in the project's Automated Perf Test | Static Camera settings. Exiting test early.");
 		EndTestFailure();
 		return;
 	}
